Slate: Moved TestSuite and ColorPicker into new module AppFramework

The AppFramework module is intended to be used for compound widgets and UI related classes that are too specific (not basic enough) for Slate, but also not Editor specific (reusable in non-Editor applications and games). The test suite has been moved in its entirety for now, but core widget specific test classes will eventually be split off and moved back into Slate, so that they can live alongside of their corresponding widgets.

Other changes:
- moved to "include what you use" scheme for SColorPicker
- broke out color picker related widgets that may be reusable
- added forward declarations to reduce header include dependencies

diff --git a/Engine/Source/Editor/Cascade/Cascade.Build.cs b/Engine/Source/Editor/Cascade/Cascade.Build.cs
--- a/Engine/Source/Editor/Cascade/Cascade.Build.cs
+++ b/Engine/Source/Editor/Cascade/Cascade.Build.cs
@@ -15,6 +15,7 @@ public class Cascade : ModuleRules
 
 		PrivateDependencyModuleNames.AddRange(
 			new string[] {
+                "AppFramework",
 				"Core",
 				"CoreUObject",
                 "InputCore",
diff --git a/Engine/Source/Editor/Cascade/Private/Cascade.cpp b/Engine/Source/Editor/Cascade/Private/Cascade.cpp
--- a/Engine/Source/Editor/Cascade/Private/Cascade.cpp
+++ b/Engine/Source/Editor/Cascade/Private/Cascade.cpp
@@ -14,7 +14,7 @@
 #include "CascadeActions.h"
 #include "Cascade.h"
 #include "PhysicsPublic.h"
-
+#include "SColorPicker.h"
 
 #include "Editor/WorkspaceMenuStructure/Public/WorkspaceMenuStructureModule.h"
 #include "Editor/PropertyEditor/Public/PropertyEditorModule.h"
diff --git a/Engine/Source/Editor/Cascade/Private/CascadeEmitterCanvasClient.cpp b/Engine/Source/Editor/Cascade/Private/CascadeEmitterCanvasClient.cpp
--- a/Engine/Source/Editor/Cascade/Private/CascadeEmitterCanvasClient.cpp
+++ b/Engine/Source/Editor/Cascade/Private/CascadeEmitterCanvasClient.cpp
@@ -16,6 +16,7 @@
 #include "Particles/ParticleModuleRequired.h"
 #include "Particles/ParticleSpriteEmitter.h"
 #include "Particles/ParticleSystem.h"
+#include "SColorPicker.h"
 
 
 FCascadeEmitterCanvasClient::FCascadeEmitterCanvasClient(TWeakPtr<FCascade> InCascade, TWeakPtr<SCascadeEmitterCanvas> InCascadeViewport)
diff --git a/Engine/Source/Editor/ContentBrowser/Private/AssetContextMenu.cpp b/Engine/Source/Editor/ContentBrowser/Private/AssetContextMenu.cpp
--- a/Engine/Source/Editor/ContentBrowser/Private/AssetContextMenu.cpp
+++ b/Engine/Source/Editor/ContentBrowser/Private/AssetContextMenu.cpp
@@ -28,6 +28,9 @@
 #include "IDocumentation.h"
 #include "EditorClassUtils.h"
 
+#include "SColorPicker.h"
+
+
 #define LOCTEXT_NAMESPACE "ContentBrowser"
 
 FAssetContextMenu::FAssetContextMenu(const TWeakPtr<SAssetView>& InAssetView)
diff --git a/Engine/Source/Editor/DetailCustomizations/Private/MathStructCustomizations.cpp b/Engine/Source/Editor/DetailCustomizations/Private/MathStructCustomizations.cpp
--- a/Engine/Source/Editor/DetailCustomizations/Private/MathStructCustomizations.cpp
+++ b/Engine/Source/Editor/DetailCustomizations/Private/MathStructCustomizations.cpp
@@ -4,9 +4,12 @@
 #include "Materials/MaterialExpressionConstant3Vector.h"
 #include "MathStructCustomizations.h"
 #include "IPropertyUtilities.h"
+#include "SColorPicker.h"
+
 
 #define LOCTEXT_NAMESPACE "FMathStructCustomization"
 
+
 TSharedRef<IPropertyTypeCustomization> FMathStructCustomization::MakeInstance() 
 {
 	return MakeShareable( new FMathStructCustomization );
diff --git a/Engine/Source/Editor/DetailCustomizations/Private/MathStructCustomizations.h b/Engine/Source/Editor/DetailCustomizations/Private/MathStructCustomizations.h
--- a/Engine/Source/Editor/DetailCustomizations/Private/MathStructCustomizations.h
+++ b/Engine/Source/Editor/DetailCustomizations/Private/MathStructCustomizations.h
@@ -178,7 +178,7 @@ private:
 	void CreateColorPicker( bool bUseAlpha, bool bOnlyRefreshOnOk );
 	
 	/** Creates a new color picker for interactively selecting color */
-	TSharedRef<SColorPicker> CreateInlineColorPicker();
+	TSharedRef<class SColorPicker> CreateInlineColorPicker();
 
 	/**
 	 * Called when the property is set from the color picker 
diff --git a/Engine/Source/Editor/DistCurveEditor/Private/SDistributionCurveEditor.cpp b/Engine/Source/Editor/DistCurveEditor/Private/SDistributionCurveEditor.cpp
--- a/Engine/Source/Editor/DistCurveEditor/Private/SDistributionCurveEditor.cpp
+++ b/Engine/Source/Editor/DistCurveEditor/Private/SDistributionCurveEditor.cpp
@@ -7,6 +7,7 @@
 #include "CurveEditorSharedData.h"
 #include "SDistributionCurveEditor.h"
 #include "Runtime/Engine/Public/Slate/SceneViewport.h"
+#include "SColorPicker.h"
 
 #define LOCTEXT_NAMESPACE "CurveEditor"
 
diff --git a/Engine/Source/Editor/GraphEditor/Private/KismetPins/SGraphPinColor.cpp b/Engine/Source/Editor/GraphEditor/Private/KismetPins/SGraphPinColor.cpp
--- a/Engine/Source/Editor/GraphEditor/Private/KismetPins/SGraphPinColor.cpp
+++ b/Engine/Source/Editor/GraphEditor/Private/KismetPins/SGraphPinColor.cpp
@@ -4,6 +4,8 @@
 #include "GraphEditorCommon.h"
 #include "NodeFactory.h"
 #include "KismetPins/SGraphPinColor.h"
+#include "SColorPicker.h"
+
 
 void SGraphPinColor::Construct(const FArguments& InArgs, UEdGraphPin* InGraphPinObj)
 {
diff --git a/Engine/Source/Editor/IntroTutorials/IntroTutorials.Build.cs b/Engine/Source/Editor/IntroTutorials/IntroTutorials.Build.cs
--- a/Engine/Source/Editor/IntroTutorials/IntroTutorials.Build.cs
+++ b/Engine/Source/Editor/IntroTutorials/IntroTutorials.Build.cs
@@ -28,7 +28,8 @@ namespace UnrealBuildTool.Rules
 			);
 
 			PrivateDependencyModuleNames.AddRange(
-				new string[] {				
+				new string[] {
+                    "AppFramework",
                     "UnrealEd",
                     "Kismet",
                     "PlacementMode",
diff --git a/Engine/Source/Editor/Kismet/Private/BlueprintDetailsCustomization.cpp b/Engine/Source/Editor/Kismet/Private/BlueprintDetailsCustomization.cpp
--- a/Engine/Source/Editor/Kismet/Private/BlueprintDetailsCustomization.cpp
+++ b/Engine/Source/Editor/Kismet/Private/BlueprintDetailsCustomization.cpp
@@ -9,7 +9,7 @@
 #include "BlueprintEditor.h"
 #include "BlueprintEditorModes.h"
 #include "Editor/PropertyEditor/Public/PropertyEditing.h"
-
+#include "SColorPicker.h"
 #include "SKismetInspector.h"
 #include "SSCSEditor.h"
 #include "SMyBlueprint.h"
diff --git a/Engine/Source/Editor/Matinee/Private/MatineeMenus.cpp b/Engine/Source/Editor/Matinee/Private/MatineeMenus.cpp
--- a/Engine/Source/Editor/Matinee/Private/MatineeMenus.cpp
+++ b/Engine/Source/Editor/Matinee/Private/MatineeMenus.cpp
@@ -45,8 +45,12 @@
 #include "Runtime/Analytics/Analytics/Public/Interfaces/IAnalyticsProvider.h"
 #include "EngineAnalytics.h"
 
+#include "SColorPicker.h"
+
+
 #define LOCTEXT_NAMESPACE "MatineeMenus"
 
+
 namespace
 {
 void* GetMatineeDialogParentWindow()
diff --git a/Engine/Source/Editor/MeshPaint/MeshPaint.Build.cs b/Engine/Source/Editor/MeshPaint/MeshPaint.Build.cs
--- a/Engine/Source/Editor/MeshPaint/MeshPaint.Build.cs
+++ b/Engine/Source/Editor/MeshPaint/MeshPaint.Build.cs
@@ -8,8 +8,9 @@ public class MeshPaint : ModuleRules
     {
 		PrivateIncludePathModuleNames.Add("AssetTools");
 
-        PrivateDependencyModuleNames.AddRange( 
-            new string[] { 
+        PrivateDependencyModuleNames.AddRange(
+            new string[] {
+                "AppFramework",
                 "Core", 
                 "CoreUObject",
                 "DesktopPlatform",
diff --git a/Engine/Source/Editor/PropertyEditor/Private/PropertyEditorModule.cpp b/Engine/Source/Editor/PropertyEditor/Private/PropertyEditorModule.cpp
--- a/Engine/Source/Editor/PropertyEditor/Private/PropertyEditorModule.cpp
+++ b/Engine/Source/Editor/PropertyEditor/Private/PropertyEditorModule.cpp
@@ -22,6 +22,8 @@
 
 #include "PropertyTableConstants.h"
 #include "SStructureDetailsView.h"
+#include "SColorPicker.h"
+
 
 IMPLEMENT_MODULE( FPropertyEditorModule, PropertyEditor );
 
diff --git a/Engine/Source/Editor/WorldBrowser/WorldBrowser.Build.cs b/Engine/Source/Editor/WorldBrowser/WorldBrowser.Build.cs
--- a/Engine/Source/Editor/WorldBrowser/WorldBrowser.Build.cs
+++ b/Engine/Source/Editor/WorldBrowser/WorldBrowser.Build.cs
@@ -17,7 +17,8 @@ public class WorldBrowser : ModuleRules
         );
      
         PrivateDependencyModuleNames.AddRange(
-            new string[] {    
+            new string[] {
+                "AppFramework",
                 "Core", 
                 "CoreUObject",
                 "RenderCore",
diff --git a/Engine/Source/Programs/SlateViewer/Private/IOS/IOSSlateViewerMain.cpp b/Engine/Source/Programs/SlateViewer/Private/IOS/IOSSlateViewerMain.cpp
--- a/Engine/Source/Programs/SlateViewer/Private/IOS/IOSSlateViewerMain.cpp
+++ b/Engine/Source/Programs/SlateViewer/Private/IOS/IOSSlateViewerMain.cpp
@@ -5,12 +5,15 @@
 #include "IOSAppDelegate.h"
 #include "IOSCommandLineHelper.h"
 #include "IOS/SlateOpenGLESView.h"
+#include "STestSuite.h"
+
 
 #define IOS_MAX_PATH 1024
 #define CMD_LINE_MAX 16384
 
 FString GSavedCommandLine;
 
+
 void FAppEntry::PreInit(IOSAppDelegate* AppDelegate, UIApplication* Application)
 {
 	// make a controller object
@@ -33,10 +36,12 @@ void FAppEntry::PreInit(IOSAppDelegate* AppDelegate, UIApplication* Application)
 	}
 }
 
+
 void FAppEntry::PlatformInit()
 {
 }
 
+
 void FAppEntry::Init()
 {
 	// start up the main loop
@@ -59,6 +64,7 @@ void FAppEntry::Init()
 	double LastTime = FPlatformTime::Seconds();
 }
 
+
 void FAppEntry::Tick()
 {
 	FSlateApplication::Get().PumpMessages();
@@ -68,11 +74,13 @@ void FAppEntry::Tick()
 	FPlatformProcess::Sleep( 0 );
 }
 
+
 void FAppEntry::Shutdown()
 {
 	FSlateApplication::Shutdown();
 }
 
+
 int main(int argc, char *argv[])
 {
 	for (int32 Option = 1; Option < argc; Option++)
diff --git a/Engine/Source/Programs/SlateViewer/Private/SlateViewerApp.cpp b/Engine/Source/Programs/SlateViewer/Private/SlateViewerApp.cpp
--- a/Engine/Source/Programs/SlateViewer/Private/SlateViewerApp.cpp
+++ b/Engine/Source/Programs/SlateViewer/Private/SlateViewerApp.cpp
@@ -2,6 +2,7 @@
 
 #include "SlateViewerApp.h"
 #include "RequiredProgramMainCPPInclude.h"
+#include "STestSuite.h"
 
 
 IMPLEMENT_APPLICATION(SlateViewer, "SlateViewer");
diff --git a/Engine/Source/Programs/SlateViewer/SlateViewer.Build.cs b/Engine/Source/Programs/SlateViewer/SlateViewer.Build.cs
--- a/Engine/Source/Programs/SlateViewer/SlateViewer.Build.cs
+++ b/Engine/Source/Programs/SlateViewer/SlateViewer.Build.cs
@@ -10,6 +10,7 @@ public class SlateViewer : ModuleRules
 
 		PrivateDependencyModuleNames.AddRange(
 			new string[] {
+                "AppFramework",
 				"Core",
 				"Projects",
 				"Slate",
@@ -23,8 +24,12 @@ public class SlateViewer : ModuleRules
 
 		if (Target.Platform == UnrealTargetPlatform.IOS)
 		{
-			PrivateDependencyModuleNames.Add("NetworkFile");
-			PrivateDependencyModuleNames.Add("StreamingFile");
+			PrivateDependencyModuleNames.AddRange(
+                new string [] {
+                    "NetworkFile",
+                    "StreamingFile"
+                }
+            );
 		}
 	}
 }
diff --git a/Engine/Source/Runtime/AppFramework/AppFramework.Build.cs b/Engine/Source/Runtime/AppFramework/AppFramework.Build.cs
new file mode 100644
--- /dev/null
+++ b/Engine/Source/Runtime/AppFramework/AppFramework.Build.cs
@@ -0,0 +1,33 @@
+// Copyright 1998-2014 Epic Games, Inc. All Rights Reserved.
+
+using UnrealBuildTool;
+
+public class AppFramework : ModuleRules
+{
+	public AppFramework(TargetInfo Target)
+	{
+		PublicDependencyModuleNames.AddRange(
+			new string[] {
+				"Core",
+				"Slate",
+				"SlateCore",
+			}
+		);
+
+		PrivateDependencyModuleNames.AddRange(
+			new string[] {
+                "InputCore",
+			}
+		);
+
+		PrivateIncludePaths.AddRange(
+			new string[] {
+				"Runtime/AppFramework/Private",
+                "Runtime/AppFramework/Private/Framework/Testing",
+				"Runtime/AppFramework/Private/Widgets",
+                "Runtime/AppFramework/Private/Widgets/Colors",
+                "Runtime/AppFramework/Private/Widgets/Testing",
+			}
+		);
+	}
+}
diff --git a/Engine/Source/Runtime/AppFramework/Private/AppFrameworkModule.cpp b/Engine/Source/Runtime/AppFramework/Private/AppFrameworkModule.cpp
new file mode 100644
--- /dev/null
+++ b/Engine/Source/Runtime/AppFramework/Private/AppFrameworkModule.cpp
@@ -0,0 +1,29 @@
+// Copyright 1998-2014 Epic Games, Inc. All Rights Reserved.
+
+#include "AppFrameworkPrivatePCH.h"
+
+
+static const FName AppFrameworkTabName("AppFramework");
+
+
+/**
+ * Implements the AppFramework module.
+ */
+class FAppFrameworkModule
+	: public IModuleInterface
+{
+public:
+
+	// IModuleInterface interface
+	
+	virtual void StartupModule( ) override
+	{
+	}
+
+	virtual void ShutdownModule( ) override
+	{
+	}
+};
+
+
+IMPLEMENT_MODULE(FAppFrameworkModule, AppFramework);
diff --git a/Engine/Source/Runtime/AppFramework/Private/AppFrameworkPrivatePCH.h b/Engine/Source/Runtime/AppFramework/Private/AppFrameworkPrivatePCH.h
new file mode 100644
--- /dev/null
+++ b/Engine/Source/Runtime/AppFramework/Private/AppFrameworkPrivatePCH.h
@@ -0,0 +1,14 @@
+// Copyright 1998-2014 Epic Games, Inc. All Rights Reserved.
+
+#pragma once
+
+#include "AppFramework.h"
+
+
+/* Private dependencies
+ *****************************************************************************/
+
+
+
+/* Private includes
+ *****************************************************************************/
diff --git a/Engine/Source/Runtime/AppFramework/Private/Framework/Testing/SLayoutExample.h b/Engine/Source/Runtime/AppFramework/Private/Framework/Testing/SLayoutExample.h
new file mode 100644
--- /dev/null
+++ b/Engine/Source/Runtime/AppFramework/Private/Framework/Testing/SLayoutExample.h
@@ -0,0 +1,6 @@
+// Copyright 1998-2014 Epic Games, Inc. All Rights Reserved.
+
+#pragma once
+
+
+APPFRAMEWORK_API TSharedRef<SWidget> MakeLayoutExample();
diff --git a/Engine/Source/Runtime/AppFramework/Private/Framework/Testing/STableViewTesting.h b/Engine/Source/Runtime/AppFramework/Private/Framework/Testing/STableViewTesting.h
new file mode 100644
--- /dev/null
+++ b/Engine/Source/Runtime/AppFramework/Private/Framework/Testing/STableViewTesting.h
@@ -0,0 +1,6 @@
+// Copyright 1998-2014 Epic Games, Inc. All Rights Reserved.
+
+#pragma once
+
+
+APPFRAMEWORK_API TSharedRef<SWidget> MakeTableViewTesting();
diff --git a/Engine/Source/Runtime/AppFramework/Private/Framework/Testing/SWidgetGallery.h b/Engine/Source/Runtime/AppFramework/Private/Framework/Testing/SWidgetGallery.h
new file mode 100644
--- /dev/null
+++ b/Engine/Source/Runtime/AppFramework/Private/Framework/Testing/SWidgetGallery.h
@@ -0,0 +1,6 @@
+// Copyright 1998-2014 Epic Games, Inc. All Rights Reserved.
+
+#pragma once
+
+/** @return a new Drag and Drop test widget */
+TSharedRef<SWidget> MakeWidgetGallery();
diff --git a/Engine/Source/Runtime/AppFramework/Private/Widgets/Colors/SComplexGradient.cpp b/Engine/Source/Runtime/AppFramework/Private/Widgets/Colors/SComplexGradient.cpp
new file mode 100644
--- /dev/null
+++ b/Engine/Source/Runtime/AppFramework/Private/Widgets/Colors/SComplexGradient.cpp
@@ -0,0 +1,63 @@
+// Copyright 1998-2014 Epic Games, Inc. All Rights Reserved.
+
+#include "AppFrameworkPrivatePCH.h"
+#include "SComplexGradient.h"
+
+
+/* SComplexGradient interface
+ *****************************************************************************/
+
+void SComplexGradient::Construct( const FArguments& InArgs )
+{
+	GradientColors = InArgs._GradientColors;
+	bHasAlphaBackground = InArgs._HasAlphaBackground.Get();
+	Orientation = InArgs._Orientation.Get();
+}
+
+
+/* SCompoundWidget overrides
+ *****************************************************************************/
+
+int32 SComplexGradient::OnPaint( const FPaintArgs& Args, const FGeometry& AllottedGeometry, const FSlateRect& MyClippingRect, FSlateWindowElementList& OutDrawElements, int32 LayerId, const FWidgetStyle& InWidgetStyle, bool bParentEnabled ) const
+{
+	ESlateDrawEffect::Type DrawEffects = (bParentEnabled && IsEnabled()) ? ESlateDrawEffect::None : ESlateDrawEffect::DisabledEffect;
+
+	if (bHasAlphaBackground)
+	{
+		const FSlateBrush* StyleInfo = FCoreStyle::Get().GetBrush("ColorPicker.AlphaBackground");
+
+		FSlateDrawElement::MakeBox(
+			OutDrawElements,
+			LayerId,
+			AllottedGeometry.ToPaintGeometry(),
+			StyleInfo,
+			MyClippingRect,
+			DrawEffects
+		);
+	}
+
+	const TArray<FLinearColor>& Colors = GradientColors.Get();
+	int32 NumColors = Colors.Num();
+
+	if (NumColors > 0)
+	{
+		TArray<FSlateGradientStop> GradientStops;
+
+		for (int32 ColorIndex = 0; ColorIndex < NumColors; ++ColorIndex)
+		{
+			GradientStops.Add(FSlateGradientStop(AllottedGeometry.Size * (float(ColorIndex) / (NumColors - 1)), Colors[ColorIndex]));
+		}
+
+		FSlateDrawElement::MakeGradient(
+			OutDrawElements,
+			LayerId + 1,
+			AllottedGeometry.ToPaintGeometry(),
+			GradientStops,
+			Orientation,
+			MyClippingRect,
+			DrawEffects
+		);
+	}
+
+	return LayerId + 1;
+}
diff --git a/Engine/Source/Runtime/AppFramework/Private/Widgets/Colors/SEyeDropperButton.h b/Engine/Source/Runtime/AppFramework/Private/Widgets/Colors/SEyeDropperButton.h
new file mode 100644
--- /dev/null
+++ b/Engine/Source/Runtime/AppFramework/Private/Widgets/Colors/SEyeDropperButton.h
@@ -0,0 +1,254 @@
+// Copyright 1998-2014 Epic Games, Inc. All Rights Reserved.
+
+#pragma once
+
+
+#define LOCTEXT_NAMESPACE "EyeDroppperButton"
+
+
+/**
+ * Class for placing a color picker eye-dropper button.
+ * A self contained until that only needs client code to set the display gamma and listen
+ * for the OnValueChanged events. It toggles the dropper when clicked.
+ * When active it captures the mouse, shows a dropper cursor and samples the pixel color constantly.
+ * It is stopped normally by hitting the ESC key.
+ */
+class SEyeDropperButton : public SButton
+{
+public:
+	DECLARE_DELEGATE(FOnDropperComplete)
+
+	SLATE_BEGIN_ARGS( SEyeDropperButton )
+		: _OnValueChanged()
+		, _OnBegin()
+		, _OnComplete()
+		, _DisplayGamma()
+		{}
+
+		/** Invoked when a new value is selected by the dropper */
+		SLATE_EVENT(FOnLinearColorValueChanged, OnValueChanged)
+
+		/** Invoked when the dropper goes from inactive to active */
+		SLATE_EVENT(FSimpleDelegate, OnBegin)
+
+		/** Invoked when the dropper goes from active to inactive */
+		SLATE_EVENT(FOnDropperComplete, OnComplete)
+
+		/** Sets the display Gamma setting - used to correct colors sampled from the screen */
+		SLATE_ATTRIBUTE(float, DisplayGamma)
+
+	SLATE_END_ARGS()
+
+	void Construct(const FArguments& InArgs )
+	{
+		OnValueChanged = InArgs._OnValueChanged;
+		OnBegin = InArgs._OnBegin;
+		OnComplete = InArgs._OnComplete;
+		DisplayGamma = InArgs._DisplayGamma;
+
+		// This is a button containing an dropper image and text that tells the user to hit Esc.
+		// Their visibility and colors are changed according to whether dropper mode is active or not.
+		SButton::Construct(
+			SButton::FArguments()
+				.ContentPadding(1.0f)
+				.OnClicked(this, &SEyeDropperButton::OnClicked)
+				[
+					SNew(SOverlay)
+
+					+ SOverlay::Slot()
+						.Padding(FMargin(1.0f, 0.0f))
+						[
+							SNew(SImage)
+								.Image(FCoreStyle::Get().GetBrush("ColorPicker.EyeDropper"))
+								.ToolTipText(LOCTEXT("EyeDropperButton_ToolTip", "Activates the eye-dropper for selecting a colored pixel from any window."))
+								.ColorAndOpacity(this, &SEyeDropperButton::GetDropperImageColor)				
+						]
+
+					+ SOverlay::Slot()
+						[
+							SNew(STextBlock)
+								.Text(LOCTEXT("EscapeCue", "Esc"))
+								.ToolTipText(LOCTEXT("EyeDropperEscapeCue_ToolTip", "Hit Escape key to stop the eye dropper"))
+								.Visibility(this, &SEyeDropperButton::GetEscapeTextVisibility)
+						]
+				]
+		);
+	}
+
+	virtual FReply OnMouseButtonDown( const FGeometry& MyGeometry, const FPointerEvent& MouseEvent ) override
+	{
+		// Clicking ANY mouse button when the dropper isn't active resets the active dropper states ready to activate
+		if (!HasMouseCapture())
+		{
+			bWasClickActivated = false;
+			bWasLeft = false;
+			bWasReEntered = false;
+		}
+
+		return SButton::OnMouseButtonDown(MyGeometry, MouseEvent);
+	}
+
+	virtual FReply OnMouseButtonUp( const FGeometry& MyGeometry, const FPointerEvent& MouseEvent ) override
+	{
+		// If a mouse click is completing and the dropper is active ALWAYS deactivate
+		bool bDeactivating = false;
+		if (bWasClickActivated)
+		{
+			bDeactivating = true;
+		}
+
+		// bWasClicked is reset here because if it is set during SButton::OnMouseButtonUp()
+		// then the button was 'clicked' according to the usual rules. We might want to capture the
+		// mouse when the button is clicked but can't do it in the Clicked callback.
+		bWasClicked = false;
+		FReply Reply = SButton::OnMouseButtonUp(MyGeometry, MouseEvent);
+
+		// Switching dropper mode off?
+		if (bDeactivating)
+		{
+			// These state flags clear dropper mode
+			bWasClickActivated = false;
+			bWasLeft = false;
+			bWasReEntered = false;
+
+			Reply.ReleaseMouseCapture();
+
+			OnComplete.ExecuteIfBound();
+		}
+		else if (MouseEvent.GetEffectingButton() == EKeys::LeftMouseButton && bWasClicked)
+		{
+			// A normal LMB mouse click on the button occurred
+			// Set the initial dropper mode state and capture the mouse
+			bWasClickActivated = true;
+			bWasLeft = false;
+			bWasReEntered = false;
+
+			OnBegin.ExecuteIfBound();
+
+			Reply.CaptureMouse(this->AsShared());
+		}
+		bWasClicked = false;
+
+		return Reply;
+	}
+
+	virtual FReply OnMouseMove( const FGeometry& MyGeometry, const FPointerEvent& MouseEvent ) override
+	{
+		// If the mouse is captured and bWasClickActivated is set then we are in dropper mode
+		if (HasMouseCapture() && bWasClickActivated)
+		{
+			if (IsMouseOver(MyGeometry, MouseEvent))
+			{
+				if (bWasLeft)
+				{
+					// Mouse is over the button having left it once
+					bWasReEntered = true;
+				}
+			}
+			else
+			{
+				// Mouse is outside the button
+				bWasLeft = true;
+				bWasReEntered = false;
+
+				// In dropper mode and outside the button - sample the pixel color and push it to the client
+				FLinearColor ScreenColor = FPlatformMisc::GetScreenPixelColor(FSlateApplication::Get().GetCursorPos(), false);
+				OnValueChanged.ExecuteIfBound(ScreenColor);
+			}
+		}
+
+		return SButton::OnMouseMove( MyGeometry, MouseEvent );
+	}
+
+	FReply OnKeyDown( const FGeometry& MyGeometry, const FKeyboardEvent& InKeyboardEvent )
+	{
+		// Escape key when in dropper mode
+		if (InKeyboardEvent.GetKey() == EKeys::Escape &&
+			HasMouseCapture() &&
+			bWasClickActivated)
+		{
+			// Clear the dropper mode states
+			bWasClickActivated = false;
+			bWasLeft = false;
+			bWasReEntered = false;
+
+			// This is needed to switch the dropper cursor off immediately so the user can see the Esc key worked
+			FSlateApplication::Get().QueryCursor();
+
+			FReply ReleaseReply = FReply::Handled().ReleaseMouseCapture();
+
+			OnComplete.ExecuteIfBound();
+
+			return ReleaseReply;
+		}
+
+		return FReply::Unhandled();
+	}
+
+	virtual FCursorReply OnCursorQuery( const FGeometry& MyGeometry, const FPointerEvent& CursorEvent ) const override
+	{
+		// Cursor is changed to the dropper when dropper mode is active and the states are correct
+		if (HasMouseCapture() &&
+			bWasClickActivated && bWasLeft && !bWasReEntered )
+		{
+			return FCursorReply::Cursor( EMouseCursor::EyeDropper );
+		}
+
+		return SButton::OnCursorQuery( MyGeometry, CursorEvent );
+	}
+
+	FReply OnClicked()
+	{
+		// Log clicks so that OnMouseButtonUp() can post process clicks
+		bWasClicked = true;
+		return FReply::Handled();
+	}
+
+private:
+
+	EVisibility GetEscapeTextVisibility() const
+	{
+		// Show the Esc key message in the button when dropper mode is active
+		if (HasMouseCapture() && bWasClickActivated)
+		{
+			return EVisibility::Visible;
+		}
+		return EVisibility::Hidden;
+	}
+
+	FSlateColor GetDropperImageColor() const
+	{
+		// Make the dropper image in the button pale when dropper mode is active
+		if (HasMouseCapture() && bWasClickActivated)
+		{
+			return FLinearColor(0.3f, 0.3f, 0.3f, 1.0f);
+		}
+		return FSlateColor::UseForeground();
+	}
+
+	bool IsMouseOver(const FGeometry& MyGeometry, const FPointerEvent& MouseEvent) const
+	{ 
+		return MyGeometry.IsUnderLocation(MouseEvent.GetScreenSpacePosition());
+	}
+
+	/** Invoked when a new value is selected by the dropper */
+	FOnLinearColorValueChanged OnValueChanged;
+
+	/** Invoked when the dropper goes from inactive to active */
+	FSimpleDelegate OnBegin;
+
+	/** Invoked when the dropper goes from active to inactive - can be used to commit colors by the owning picker */
+	FOnDropperComplete OnComplete;
+
+	/** Sets the display Gamma setting - used to correct colors sampled from the screen */
+	TAttribute<float> DisplayGamma;
+
+	// Dropper states
+	bool bWasClicked;
+	bool bWasClickActivated;
+	bool bWasLeft;
+	bool bWasReEntered;
+};
+
+
+#undef LOCTEXT_NAMESPACE
diff --git a/Engine/Source/Runtime/AppFramework/Private/Widgets/Colors/SSimpleGradient.cpp b/Engine/Source/Runtime/AppFramework/Private/Widgets/Colors/SSimpleGradient.cpp
new file mode 100644
--- /dev/null
+++ b/Engine/Source/Runtime/AppFramework/Private/Widgets/Colors/SSimpleGradient.cpp
@@ -0,0 +1,57 @@
+// Copyright 1998-2014 Epic Games, Inc. All Rights Reserved.
+
+#include "AppFrameworkPrivatePCH.h"
+#include "SSimpleGradient.h"
+
+
+/* SSimpleGradient interface
+ *****************************************************************************/
+
+void SSimpleGradient::Construct( const FArguments& InArgs )
+{
+	StartColor = InArgs._StartColor;
+	EndColor = InArgs._EndColor;
+	bHasAlphaBackground = InArgs._HasAlphaBackground.Get();
+	Orientation = InArgs._Orientation.Get();
+}
+
+
+/* SCompoundWidget overrides
+ *****************************************************************************/
+
+int32 SSimpleGradient::OnPaint( const FPaintArgs& Args, const FGeometry& AllottedGeometry, const FSlateRect& MyClippingRect, FSlateWindowElementList& OutDrawElements, int32 LayerId, const FWidgetStyle& InWidgetStyle, bool bParentEnabled ) const
+{
+	const ESlateDrawEffect::Type DrawEffects = this->ShouldBeEnabled(bParentEnabled) ? ESlateDrawEffect::None : ESlateDrawEffect::DisabledEffect;
+
+	if (bHasAlphaBackground)
+	{
+		const FSlateBrush* StyleInfo = FCoreStyle::Get().GetBrush("ColorPicker.AlphaBackground");
+
+		FSlateDrawElement::MakeBox(
+			OutDrawElements,
+			LayerId,
+			AllottedGeometry.ToPaintGeometry(),
+			StyleInfo,
+			MyClippingRect,
+			DrawEffects
+		);
+	}
+
+	TArray<FSlateGradientStop> GradientStops;
+
+	GradientStops.Add(FSlateGradientStop(FVector2D::ZeroVector, StartColor.Get()));
+	GradientStops.Add(FSlateGradientStop(AllottedGeometry.Size, EndColor.Get()));
+
+	FSlateDrawElement::MakeGradient(
+		OutDrawElements,
+		LayerId + 1,
+		AllottedGeometry.ToPaintGeometry(),
+		GradientStops,
+		Orientation,
+		MyClippingRect,
+		DrawEffects,
+		false
+	);
+
+	return LayerId + 1;
+}
diff --git a/Engine/Source/Runtime/AppFramework/Public/AppFramework.h b/Engine/Source/Runtime/AppFramework/Public/AppFramework.h
new file mode 100644
--- /dev/null
+++ b/Engine/Source/Runtime/AppFramework/Public/AppFramework.h
@@ -0,0 +1,18 @@
+// Copyright 1998-2014 Epic Games, Inc. All Rights Reserved.
+
+#pragma once
+
+
+/* Public dependencies
+ *****************************************************************************/
+
+#include "Core.h"
+#include "ModuleInterface.h"
+#include "Slate.h"
+
+
+/* Public includes
+ *****************************************************************************/
+
+#include "SColorThemes.h"
+#include "SColorPicker.h"
diff --git a/Engine/Source/Runtime/AppFramework/Public/Widgets/Colors/SColorPicker.h b/Engine/Source/Runtime/AppFramework/Public/Widgets/Colors/SColorPicker.h
new file mode 100644
--- /dev/null
+++ b/Engine/Source/Runtime/AppFramework/Public/Widgets/Colors/SColorPicker.h
@@ -0,0 +1,514 @@
+// Copyright 1998-2014 Epic Games, Inc. All Rights Reserved.
+
+#pragma once
+
+
+/** Called when the color picker cancel button is pressed */
+DECLARE_DELEGATE_OneParam(FOnColorPickerCancelled, FLinearColor);
+
+
+/**
+ * Enumerates color channels (do not reorder).
+ */
+enum class EColorPickerChannels
+{
+	Red,
+	Green,
+	Blue,
+	Alpha,
+	Hue,
+	Saturation,
+	Value
+};
+
+
+/**
+ * Enumerates color picker modes.
+ */
+enum class EColorPickerModes
+{
+	Spectrum,
+	Wheel
+};
+
+
+/**
+ * Struct for holding individual pointers to float values.
+ */
+struct FColorChannels
+{
+	FColorChannels()
+	{
+		Red = Green = Blue = Alpha = nullptr;
+	}
+
+	float* Red;
+	float* Green;
+	float* Blue;
+	float* Alpha;
+};
+
+
+/**
+ * Class for placing a color picker. If all you need is a standalone color picker,
+ * use the functions OpenColorPicker and DestroyColorPicker, since they hold a static
+ * instance of the color picker.
+ */
+class APPFRAMEWORK_API SColorPicker
+	: public SCompoundWidget
+{
+public:
+
+	SLATE_BEGIN_ARGS(SColorPicker)
+		: _TargetColorAttribute(FLinearColor(ForceInit))
+		, _TargetFColors()
+		, _TargetLinearColors()
+		, _TargetColorChannels()
+		, _UseAlpha(true)
+		, _OnlyRefreshOnMouseUp(false)
+		, _OnlyRefreshOnOk(false)
+		, _OnColorCommitted()
+		, _PreColorCommitted()
+		, _OnColorPickerCancelled()
+		, _OnColorPickerWindowClosed()
+		, _OnInteractivePickBegin()
+		, _OnInteractivePickEnd()
+		, _ParentWindow()
+		, _DisplayGamma(2.2f)
+		, _DisplayInlineVersion(false)
+	{ }
+		
+		/** The color that is being targeted as a TAttribute */
+		SLATE_ATTRIBUTE(FLinearColor, TargetColorAttribute)
+		
+		/** An array of color pointers this color picker targets */
+		SLATE_ATTRIBUTE(TArray<FColor*>, TargetFColors)
+		
+		/** An array of linear color pointers this color picker targets */
+		SLATE_ATTRIBUTE(TArray<FLinearColor*>, TargetLinearColors)
+		
+		/**
+		 * An array of color pointer structs this color picker targets
+		 * Only to keep compatibility with wx. Should be removed once wx is gone.
+		 */
+		SLATE_ATTRIBUTE(TArray<FColorChannels>, TargetColorChannels)
+
+		/** Whether the ability to pick the alpha value is enabled */
+		SLATE_ATTRIBUTE(bool, UseAlpha)
+		
+		/** Prevents immediate refreshs for performance reasons. */
+		SLATE_ATTRIBUTE(bool, OnlyRefreshOnMouseUp)
+
+		/** Prevents multiple refreshes when requested. */
+		SLATE_ATTRIBUTE(bool, OnlyRefreshOnOk)
+		
+		/** The event called when the color is committed */
+		SLATE_EVENT(FOnLinearColorValueChanged, OnColorCommitted)
+
+		/** The event called before the color is committed */
+		SLATE_EVENT(FOnLinearColorValueChanged, PreColorCommitted)
+		
+		/** The event called when the color picker cancel button is pressed */
+		SLATE_EVENT(FOnColorPickerCancelled, OnColorPickerCancelled)
+
+		/** The event called when the color picker parent window is closed */
+		SLATE_EVENT(FOnWindowClosed, OnColorPickerWindowClosed)
+
+		/** The event called when a slider drag, color wheel drag or dropper grab starts */
+		SLATE_EVENT(FSimpleDelegate, OnInteractivePickBegin)
+
+		/** The event called when a slider drag, color wheel drag or dropper grab finishes */
+		SLATE_EVENT(FSimpleDelegate, OnInteractivePickEnd)
+
+		/** A pointer to the parent window */
+		SLATE_ATTRIBUTE(TSharedPtr<SWindow>, ParentWindow)
+
+		/** Sets the display Gamma setting - used to correct colors sampled from the screen */
+		SLATE_ATTRIBUTE(float, DisplayGamma)
+
+		/** If true, this color picker will be a stripped down version of the full color picker */
+		SLATE_ARGUMENT(bool, DisplayInlineVersion)
+
+	SLATE_END_ARGS()
+	
+	/** A default window size for the color picker which looks nice */
+	static const FVector2D DEFAULT_WINDOW_SIZE;
+
+	/**	Destructor. */
+	~SColorPicker();
+
+public:
+
+	/**
+	 * Construct the widget
+	 *
+	 * @param InArgs Declaration from which to construct the widget.
+	 */
+	void Construct(const FArguments& InArgs );
+
+protected:
+
+	/** Backup all the colors that are being modified */
+	void BackupColors();
+
+	/** Restore all the modified colors to their original state */
+	void RestoreColors();
+
+	/** Set all the colors to this new color */
+	void SetColors( const FLinearColor& InColor );
+
+	bool ApplyNewTargetColor( bool bForceUpdate = false );
+
+	void GenerateDefaultColorPickerContent( bool bAdvancedSectionExpanded );
+	void GenerateInlineColorPickerContent();
+
+	FLinearColor GetCurrentColor() const
+	{
+		return CurrentColorHSV;
+	}
+
+	/** Calls the user defined delegate for when the color changes are discarded */
+	void DiscardColor();
+
+	/** Sets new color in ether RGB or HSV */
+	bool SetNewTargetColorRGB( const FLinearColor& NewValue, bool bForceUpdate = false );
+	bool SetNewTargetColorHSV( const FLinearColor& NewValue, bool bForceUpdate = false );
+
+    void UpdateColorPick();
+	void UpdateColorPickMouseUp();
+
+	void BeginAnimation(FLinearColor Start, FLinearColor End);
+
+	void HideSmallTrash();
+	void ShowSmallTrash();
+
+	/** Cycles the color picker's mode. */
+	void CycleMode( );
+
+	/**
+	 * Creates a color slider widget for the specified channel.
+	 *
+	 * @param Channel The color channel to create the widget for.
+	 * @return The new slider.
+	 */
+	TSharedRef<SWidget> MakeColorSlider( EColorPickerChannels Channel ) const;
+
+	/**
+	 * Creates a color spin box widget for the specified channel.
+	 *
+	 * @param Channel The color channel to create the widget for.
+	 * @return The new spin box.
+	 */
+	TSharedRef<SWidget> MakeColorSpinBox( EColorPickerChannels Channel ) const;
+
+	/**
+	 * Creates the color preview box widget.
+	 *
+	 * @return The new color preview box.
+	 */
+	TSharedRef<SWidget> MakeColorPreviewBox( ) const;
+
+protected:	
+
+	// SWidget interface
+
+	virtual void Tick( const FGeometry& AllottedGeometry, const double InCurrentTime, const float InDeltaTime ) override;
+
+private:
+
+	// Callback for getting the end color of a color spin box gradient.
+	FLinearColor GetGradientEndColor( EColorPickerChannels Channel ) const;
+
+	// Callback for getting the start color of a color spin box gradient.
+	FLinearColor GetGradientStartColor( EColorPickerChannels Channel ) const;
+
+	// Callback for handling expansion of the 'Advanced' area.
+	void HandleAdvancedAreaExpansionChanged( bool Expanded );
+
+	// Callback for getting the visibility of the alpha channel portion in color blocks.
+	EVisibility HandleAlphaColorBlockVisibility( ) const;
+
+	// Callback for clicking the Cancel button.
+	FReply HandleCancelButtonClicked();
+
+	// Callback for pressing a mouse button in the color area.
+	FReply HandleColorAreaMouseDown( const FGeometry& MyGeometry, const FPointerEvent& MouseEvent );
+
+	// Callback for clicking the color picker mode button.
+	FReply HandleColorPickerModeButtonClicked( );
+
+	// Callback for getting the visibility of the given color picker mode.
+	EVisibility HandleColorPickerModeVisibility( EColorPickerModes Mode ) const;
+
+	// Callback for getting the end color of a color slider.
+	FLinearColor HandleColorSliderEndColor( EColorPickerChannels Channel ) const;
+
+	// Callback for getting the start color of a color slider.
+	FLinearColor HandleColorSliderStartColor( EColorPickerChannels Channel ) const;
+
+	// Callback for value changes in the color spectrum picker.
+	void HandleColorSpectrumValueChanged( FLinearColor NewValue );
+
+	// Callback for getting the value of a color spin box.
+	float HandleColorSpinBoxValue( EColorPickerChannels Channel ) const;
+
+	// Callback for value changes in a color spin box.
+	void HandleColorSpinBoxValueChanged( float NewValue, EColorPickerChannels Channel );
+
+	// Callback for completed eye dropper interactions.
+	void HandleEyeDropperButtonComplete( );
+
+	// Callback for getting the text in the hex box.
+	FText HandleHexBoxText( ) const;
+
+	// Callback for committed text in the hex input box.
+	void HandleHexInputTextCommitted( const FText& Text, ETextCommit::Type CommitType );
+
+	// Callback for changing the HSV value of the current color.
+	void HandleHSVColorChanged( FLinearColor NewValue );
+
+	// Callback for when interactive user input begins.
+	void HandleInteractiveChangeBegin( );
+
+	// Callback for when interactive user input ends.
+	void HandleInteractiveChangeEnd( );
+
+	// Callback for when interactive user input ends.
+	void HandleInteractiveChangeEnd( float NewValue );
+
+	// Callback for clicking the new color preview block.
+	FReply HandleNewColorBlockMouseButtonDown( const FGeometry& MyGeometry, const FPointerEvent& MouseEvent, bool bCheckAlpha );
+
+	// Callback for clicking the OK button.
+	FReply HandleOkButtonClicked( );
+
+	// Callback for clicking the old color preview block.
+	FReply HandleOldColorBlockMouseButtonDown( const FGeometry& MyGeometry, const FPointerEvent& MouseEvent, bool bCheckAlpha );
+
+	// Callback for checking whether sRGB colors should be rendered.
+	bool HandleColorPickerUseSRGB( ) const;
+
+	// Callback for when the parent window has been closed.
+	void HandleParentWindowClosed( const TSharedRef<SWindow>& Window );
+
+	// Callback for changing the RGB value of the current color.
+	void HandleRGBColorChanged( FLinearColor NewValue );
+
+	// Callback for changing the checked state of the sRGB check box.
+	void HandleSRGBCheckBoxCheckStateChanged( ESlateCheckBoxState::Type InIsChecked );
+
+	// Callback for determining whether the sRGB check box should be checked.
+	ESlateCheckBoxState::Type HandleSRGBCheckBoxIsChecked( ) const;
+
+	// Callback for selecting a color in the color theme bar.
+	void HandleThemeBarColorSelected( FLinearColor NewValue );
+
+	// Callback for getting the theme bar's color theme.
+	TSharedPtr<class FColorTheme> HandleThemeBarColorTheme( ) const;
+
+	// Callback for getting the visibility of the theme bar hint text.
+	EVisibility HandleThemeBarHintVisibility( ) const;
+
+	// Callback for determining whether the theme bar should display the alpha channel.
+	bool HandleThemeBarUseAlpha( ) const;
+
+	// Callback for theme viewer changes.
+	void HandleThemesViewerThemeChanged( );
+
+private:
+	
+	/** The color that is being targeted as a TAttribute */
+	TAttribute<FLinearColor> TargetColorAttribute;
+
+	/** The current color being picked in HSV */
+	FLinearColor CurrentColorHSV;
+
+	/** The current color being picked in RGB */
+	FLinearColor CurrentColorRGB;
+	
+	/** The old color to be changed in HSV */
+	FLinearColor OldColor;
+	
+	/** Color end point to animate to */
+	FLinearColor ColorEnd;
+
+	/** Color start point to animate from */
+	FLinearColor ColorBegin;
+
+	// Holds the color picker's mode.
+	EColorPickerModes CurrentMode;
+
+	/** Time, used for color animation */
+	float CurrentTime;
+
+	/** The max time allowed for updating before we shut off auto-updating */
+	static const double MAX_ALLOWED_UPDATE_TIME;
+
+	/** If true, then the performance is too bad to have auto-updating */
+	bool bPerfIsTooSlowToUpdate;
+
+	/** An array of color pointers this color picker targets */
+	TArray<FColor*> TargetFColors;
+	
+	/** An array of linear color pointers this color picker targets */
+	TArray<FLinearColor*> TargetLinearColors;
+	
+	/**
+	 * An array of color pointer structs this color picker targets
+	 * Only to keep compatibility with wx. Should be removed once wx is gone.
+	 */
+	TArray<FColorChannels> TargetColorChannels;
+
+	/** Backups of the TargetFColors */
+	TArray<FColor> OldTargetFColors;
+
+	/** Backups of the TargetLinearColors */
+	TArray<FLinearColor> OldTargetLinearColors;
+
+	/** Backups of the TargetColorChannels */
+	TArray<FLinearColor> OldTargetColorChannels;
+
+	/** Whether or not the color uses Alpha or not */
+	TAttribute<bool> bUseAlpha;
+
+	/** Prevents immediate refresh for performance reasons. */
+	bool bOnlyRefreshOnMouseUp;
+
+	/** Prevents multiple refreshes when requested. */
+	bool bOnlyRefreshOnOk;
+
+	/** true if the picker was closed via the OK or Cancel buttons, false otherwise */
+	bool bClosedViaOkOrCancel;
+
+	/** Pointer to the parent window, so we know to destroy it when done */
+	TWeakPtr<SWindow> ParentWindowPtr;
+
+	/** The widget which holds the currently selected theme */
+	TSharedPtr<class SThemeColorBlocksBar> CurrentThemeBar;
+
+	/** Widget which is either the button to show the color themes viewer, or to be a color trash */
+	TSharedPtr<SBorder> ColorThemeButtonOrSmallTrash;
+
+	/** The button to show the color themes viewer */
+	TSharedPtr<class SComboButton> ColorThemeComboButton;
+
+	/** The small color trash shown in place of the combo button */
+	TSharedPtr<SWidget> SmallTrash;
+
+	/** Sets the display Gamma setting - used to correct colors sampled from the screen */
+	TAttribute<float> DisplayGamma;
+
+	/** True if this color picker is an inline color picker */
+	bool bColorPickerIsInlineVersion;
+
+	/** Tracks whether the user is moving a value spin box, the color wheel and the dropper */
+	bool bIsInteractive;
+
+private:
+
+	/** Invoked when a new value is selected on the color wheel */
+	FOnLinearColorValueChanged OnColorCommitted;
+
+	/** Invoked before a new value is selected on the color wheel */
+	FOnLinearColorValueChanged PreColorCommitted;
+
+	/** Invoked when the color picker cancel button is pressed */
+	FOnColorPickerCancelled OnColorPickerCancelled;
+
+	/** Invoked when a slider drag, color wheel drag or dropper grab starts */
+	FSimpleDelegate OnInteractivePickBegin;
+
+	/** Invoked when a slider drag, color wheel drag or dropper grab finishes */
+	FSimpleDelegate OnInteractivePickEnd;
+
+	/** Invoked when the color picker window closes. */
+	FOnWindowClosed OnColorPickerWindowClosed;
+
+private:
+
+	/** A static pointer to the global color themes viewer */
+	static TSharedPtr<class SColorThemesViewer> ColorThemesViewer;
+};
+
+
+struct FColorPickerArgs
+{
+	/** Whether or not the new color picker is modal. */
+	bool bIsModal;
+
+	/** The parent for the new color picker window */
+	TSharedPtr<const SWidget> ParentWidget;
+
+	/** Whether or not to enable the alpha slider. */
+	bool bUseAlpha;
+
+	/** Whether to disable the refresh except on mouse up for performance reasons. */
+	bool bOnlyRefreshOnMouseUp;
+
+	/** Whether to disable the refresh until the picker closes. */
+	bool bOnlyRefreshOnOk;
+
+	/** The current display gamma used to correct colors picked from the display. */
+	TAttribute<float> DisplayGamma;
+
+	/** An array of FColors to target. */
+	const TArray<FColor*>* ColorArray;
+
+	/** An array of FLinearColors to target. */
+	const TArray<FLinearColor*>* LinearColorArray;
+
+	/** An array of FColorChannels to target. (deprecated now that wx is gone?) */
+	const TArray<FColorChannels>* ColorChannelsArray;
+
+	/** A delegate to be called when the color changes. */
+	FOnLinearColorValueChanged OnColorCommitted;
+
+	/** A delegate to be called before the color change is committed. */
+	FOnLinearColorValueChanged PreColorCommitted;
+
+	/** A delegate to be called when the color picker window closes. */
+	FOnWindowClosed OnColorPickerWindowClosed;
+
+	/** A delegate to be called when the color picker cancel button is pressed */
+	FOnColorPickerCancelled OnColorPickerCancelled;
+
+	/** A delegate to be called when a slider drag, color wheel drag or dropper grab starts */
+	FSimpleDelegate OnInteractivePickBegin;
+
+	/** A delegate to be called when a slider drag, color wheel drag or dropper grab finishes */
+	FSimpleDelegate OnInteractivePickEnd;
+
+	/** Overrides the initial color set on the color picker. */
+	FLinearColor InitialColorOverride;
+
+	FColorPickerArgs()
+		: bIsModal(false)
+		, ParentWidget(nullptr)
+		, bUseAlpha(false)
+		, bOnlyRefreshOnMouseUp(false)
+		, bOnlyRefreshOnOk(false)
+		, DisplayGamma(2.2f)
+		, ColorArray(nullptr)
+		, LinearColorArray(nullptr)
+		, ColorChannelsArray(nullptr)
+		, OnColorCommitted()
+		, PreColorCommitted()
+		, OnColorPickerWindowClosed()
+		, OnColorPickerCancelled()
+		, OnInteractivePickBegin()
+		, OnInteractivePickEnd()
+		, InitialColorOverride()
+	{ }
+};
+
+
+/**
+ * Opens up the static color picker, destroying any previously existing one.
+ */
+APPFRAMEWORK_API bool OpenColorPicker( const FColorPickerArgs& Args );
+
+/**
+ * Destroys the current color picker. Necessary if the values the color picker
+ * currently targets become invalid.
+ */
+APPFRAMEWORK_API void DestroyColorPicker();
diff --git a/Engine/Source/Runtime/AppFramework/Public/Widgets/Colors/SColorThemes.h b/Engine/Source/Runtime/AppFramework/Public/Widgets/Colors/SColorThemes.h
new file mode 100644
--- /dev/null
+++ b/Engine/Source/Runtime/AppFramework/Public/Widgets/Colors/SColorThemes.h
@@ -0,0 +1,648 @@
+// Copyright 1998-2014 Epic Games, Inc. All Rights Reserved.
+
+#pragma once
+
+
+/**
+ * A Color Theme is a name and an array of Colors.
+ * It also holds and array of refresh callbacks which it calls every time it changes at all.
+ */
+class FColorTheme
+{
+public:
+
+	FColorTheme(const FString& InName = TEXT(""), const TArray< TSharedPtr<FLinearColor> >& InColors = TArray< TSharedPtr<FLinearColor> >());
+
+	/** Get a list of all the colors in the theme */
+	const TArray< TSharedPtr<FLinearColor> >& GetColors() const { return Colors; }
+
+	/** Insert a color at a specific point in the list and broadcast change */
+	void InsertNewColor(TSharedPtr<FLinearColor> InColor, int32 InsertPosition);
+
+	/** Check to see if a color is already present in the list */
+	int32 FindApproxColor(const FLinearColor& InColor, float Tolerance = KINDA_SMALL_NUMBER) const;
+
+	/** Remove all colors from the list, broadcast change */
+	void RemoveAll();
+
+	/** Remove color at index from the list, broadcast change */
+	void RemoveColor(int32 ColorIndex);
+
+	/** Remove specific color from the list, broadcast change */
+	int32 RemoveColor(const TSharedPtr<FLinearColor> InColor);
+	
+	FString Name;
+
+	DECLARE_EVENT( FColorTheme, FRefreshEvent );
+	FRefreshEvent& OnRefresh() { return RefreshEvent; }
+
+private:
+
+	TArray< TSharedPtr<FLinearColor> > Colors;
+
+	FRefreshEvent RefreshEvent;
+};
+
+
+/**
+ * The SColorTrash is a multipurpose widget which allows FColorDragDrops
+ * to be dropped on to to be 
+ */
+class SColorTrash : public SCompoundWidget
+{
+public:
+
+	SLATE_BEGIN_ARGS(SColorTrash)
+		: _UsesSmallIcon(false)
+	{ }
+
+		SLATE_ATTRIBUTE(bool, UsesSmallIcon)
+	SLATE_END_ARGS()
+
+	void Construct( const FArguments& InArgs );
+
+protected:	
+
+	/**
+	 * Called during drag and drop when the drag enters a widget.
+	 *
+	 * @param MyGeometry      The geometry of the widget receiving the event.
+	 * @param DragDropEvent   The drag and drop event.
+	 * @return A reply that indicated whether this event was handled.
+	 */
+	virtual void OnDragEnter( const FGeometry& MyGeometry, const FDragDropEvent& DragDropEvent ) override;
+
+	/**
+	 * Called during drag and drop when the drag leaves a widget.
+	 *
+	 * @param DragDropEvent   The drag and drop event.
+	 * @return A reply that indicated whether this event was handled.
+	 */
+	virtual void OnDragLeave( const FDragDropEvent& DragDropEvent ) override;
+
+	/**
+	 * Called when the user is dropping something onto a widget; terminates drag and drop.
+	 *
+	 * @param MyGeometry      The geometry of the widget receiving the event.
+	 * @param DragDropEvent   The drag and drop event.
+	 * @return A reply that indicated whether this event was handled.
+	 */
+	virtual FReply OnDrop( const FGeometry& MyGeometry, const FDragDropEvent& DragDropEvent ) override;
+	
+	const FSlateBrush* GetBorderStyle() const;
+
+private:
+
+	/** Determines whether to draw the border to show activation */
+	bool bBorderActivated;
+};
+
+
+class SThemeColorBlocksBar;
+
+
+/**
+ * SThemeColorBlocks are Color Blocks which point to a Color in a ColorTheme.
+ * They can be dragged and dropped, and clicking on one in the Color Picker will
+ * give the color that they point to.
+ */
+class SThemeColorBlock : public SCompoundWidget
+{
+public:
+
+	SLATE_BEGIN_ARGS(SThemeColorBlock)
+		: _Color()
+		, _OnSelectColor()
+		, _Parent()
+		, _ShowTrashCallback()
+		, _HideTrashCallback()
+		, _UseSRGB()
+		, _UseAlpha()
+	{ }
+
+		/** A pointer to the color this block uses */
+		SLATE_ATTRIBUTE(TSharedPtr<FLinearColor>, Color)
+		
+		/** Event called when this block is clicked */
+		SLATE_EVENT(FOnLinearColorValueChanged, OnSelectColor)
+
+		/** A pointer to the theme color blocks bar that is this block's origin */
+		SLATE_ATTRIBUTE(TSharedPtr<SThemeColorBlocksBar>, Parent)
+		
+		/** Callback to pass down to the FColorDragDrop for it to show the trash */
+		SLATE_EVENT(FSimpleDelegate, ShowTrashCallback)
+		/** Callback to pass down to the FColorDragDrop for it to hide the trash */
+		SLATE_EVENT(FSimpleDelegate, HideTrashCallback)
+		
+		/** Whether to display sRGB color */
+		SLATE_ATTRIBUTE(bool, UseSRGB)
+
+		/** Whether the ability to pick the alpha value is enabled */
+		SLATE_ATTRIBUTE(bool, UseAlpha)
+
+	SLATE_END_ARGS()
+
+	/**
+	 * Construct the widget
+	 *
+	 * @param InArgs   Declaration from which to construct the widget.
+	 */
+	void Construct(const FArguments& InArgs );
+	
+private:
+
+	virtual FReply OnMouseButtonDown( const FGeometry& MyGeometry, const FPointerEvent& MouseEvent ) override;
+	virtual FReply OnMouseButtonUp( const FGeometry& MyGeometry, const FPointerEvent& MouseEvent ) override;
+	virtual FReply OnDragDetected( const FGeometry& MyGeometry, const FPointerEvent& MouseEvent ) override;
+
+	
+	FLinearColor GetColor() const;
+
+	FSlateColor HandleBorderColor() const;
+	const FSlateBrush* HandleBorderImage() const;
+
+	FText GetRedText() const;
+	FText GetGreenText() const;
+	FText GetBlueText() const;
+	FText GetAlphaText() const;
+	FText GetHueText() const;
+	FText GetSaturationText() const;
+	FText GetValueText() const;
+
+	/**
+	 * Function for formatting text for the tooltip which has limited space.
+	 */
+	FText FormatToolTipText(const FText& ColorIdentifier, float Value) const;
+	
+	bool OnReadIgnoreAlpha() const;
+	bool OnReadShowBackgroundForAlpha() const;
+	EVisibility OnGetAlphaVisibility() const;
+
+	/** A pointer to the color this block uses */
+	TWeakPtr<FLinearColor> ColorPtr;
+	
+	/** A pointer to the theme color blocks bar that is this block's origin */
+	TWeakPtr<SThemeColorBlocksBar> ParentPtr;
+	
+	/** Event called when this block is clicked */
+	FOnLinearColorValueChanged OnSelectColor;
+	
+	/** Callback to pass down to the FColorDragDrop for it to show the trash */
+	FSimpleDelegate ShowTrashCallback;
+	/** Callback to pass down to the FColorDragDrop for it to hide the trash */
+	FSimpleDelegate HideTrashCallback;
+	
+	/** Whether to use display sRGB color */
+	TAttribute<bool> bUseSRGB;
+
+	/** Whether or not the color uses Alpha or not */
+	TAttribute<bool> bUseAlpha;
+
+	float DistanceDragged;
+};
+
+
+/**
+ * SThemeColorBlocksBars are panels for dragging and dropping SColorThemeBlocks onto and off of.
+ */
+class SThemeColorBlocksBar : public SPanel
+{
+public:
+
+	SLATE_BEGIN_ARGS(SThemeColorBlocksBar)
+		: _ColorTheme()
+		, _OnSelectColor()
+		, _ShowTrashCallback()
+		, _HideTrashCallback()
+		, _EmptyText()
+		, _UseSRGB()
+		, _UseAlpha()
+	{ }
+
+		/** A pointer to the color theme that this bar should display */
+		SLATE_ATTRIBUTE(TSharedPtr<FColorTheme>, ColorTheme)
+
+		/** Event called when a color block is clicked */
+		SLATE_EVENT(FOnLinearColorValueChanged, OnSelectColor)
+
+		/** Callback to pass down to the FColorDragDrop for it to show the trash */
+		SLATE_EVENT(FSimpleDelegate, ShowTrashCallback)
+		/** Callback to pass down to the FColorDragDrop for it to hide the trash */
+		SLATE_EVENT(FSimpleDelegate, HideTrashCallback)
+		
+		/** Specify what the bar should display when no colors are present */
+		SLATE_TEXT_ARGUMENT(EmptyText)
+
+		/** Whether to display sRGB color */
+		SLATE_ATTRIBUTE(bool, UseSRGB)
+
+		/** Whether the ability to pick the alpha value is enabled */
+		SLATE_ATTRIBUTE(bool, UseAlpha)
+
+	SLATE_END_ARGS()
+
+	SThemeColorBlocksBar();
+
+	/**
+	 * Most panels do not create widgets as part of their implementation, so
+	 * they do not need to implement a Construct()
+	 */
+	void Construct(const FArguments& InArgs);
+	
+	/**
+	 * Panels arrange their children in a space described by the AllottedGeometry parameter. The results of the arrangement
+	 * should be returned by appending a FArrangedWidget pair for every child widget.
+	 *
+	 * @param AllottedGeometry    The geometry allotted for this widget by its parent.
+	 * @param ArrangedChildren    The array to which to add the WidgetGeometries that represent the arranged children.
+	 */
+	virtual void OnArrangeChildren( const FGeometry& AllottedGeometry, FArrangedChildren& ArrangedChildren ) const override;
+
+	/**
+	 * A Panel's desired size in the space required to arrange of its children on the screen while respecting all of
+	 * the children's desired sizes and any layout-related options specified by the user. See StackPanel for an example.
+	 *
+	 * @return The desired size.
+	 */
+	virtual FVector2D ComputeDesiredSize() const override;
+
+	/**
+	 * All widgets must provide a way to access their children in a layout-agnostic way.
+	 * Panels store their children in Slots, which creates a dilemma. Most panels
+	 * can store their children in a TPanelChildren<Slot>, where the Slot class
+	 * provides layout information about the child it stores. In that case
+	 * GetChildren should simply return the TPanelChildren<Slot>.
+	 */
+	virtual FChildren* GetChildren() override;
+	
+	/**
+	 * Called during drag and drop when the drag enters a widget.
+	 *
+	 * @param MyGeometry      The geometry of the widget receiving the event.
+	 * @param DragDropEvent   The drag and drop event.
+	 * @return A reply that indicated whether this event was handled.
+	 */
+	virtual void OnDragEnter( const FGeometry& MyGeometry, const FDragDropEvent& DragDropEvent ) override;
+	
+	/**
+	 * Called during drag and drop when the drag leaves a widget.
+	 *
+	 * @param DragDropEvent   The drag and drop event.
+	 * @return A reply that indicated whether this event was handled.
+	 */
+	virtual void OnDragLeave( const FDragDropEvent& DragDropEvent ) override;
+
+	/**
+	 * Called during drag and drop when the the mouse is being dragged over a widget.
+	 *
+	 * @param MyGeometry      The geometry of the widget receiving the event.
+	 * @param DragDropEvent   The drag and drop event.
+	 * @return A reply that indicated whether this event was handled.
+	 */
+	virtual FReply OnDragOver( const FGeometry& MyGeometry, const FDragDropEvent& DragDropEvent ) override;
+
+	/**
+	 * Called when the user is dropping something onto a widget; terminates drag and drop.
+	 *
+	 * @param MyGeometry      The geometry of the widget receiving the event.
+	 * @param DragDropEvent   The drag and drop event.
+	 * @return A reply that indicated whether this event was handled.
+	 */
+	virtual FReply OnDrop( const FGeometry& MyGeometry, const FDragDropEvent& DragDropEvent ) override;
+
+	/**
+	 * Adds a new color block to the Bar
+	 * @param Color				The	color the new block will be
+	 * @param InsertPosition	The location the block will be at
+	 */
+	void AddNewColorBlock(FLinearColor Color, int32 InsertPosition);
+
+	/**
+	 * @param ColorToRemove		The pointer to the color that should be removed
+	 *
+	 * @return					The index of the removed color block, INDEX_NONE if it can't be found
+	 */
+	int32 RemoveColorBlock(TSharedPtr<FLinearColor> ColorToRemove);
+
+	/**
+	 * Refresh Callbacks are used by Color Themes to notify widgets to refresh
+	 * This removes a callback from the Color Theme that this bar has a pointer to
+	 */
+	void RemoveRefreshCallback();
+	/**
+	 * Adds a callback to the Color Theme that this bar has a pointer to
+	 */
+	void AddRefreshCallback();
+
+	/**
+	 * Rebuilds the entire bar, regenerating all the constituent color blocks
+	 */
+	void Refresh();
+
+	void SetPlaceholderGrabOffset(FVector2D GrabOffset);
+	
+private:
+
+	/**
+	 * Destroys the placeholder block in this widget
+	 */
+	void DestroyPlaceholders();
+
+	/** The children blocks of this panel */
+	TSlotlessChildren<SThemeColorBlock> ColorBlocks;
+
+	/**
+	 * A placeholder child, which exists as a "preview" to show what would happen when dropping the current
+	 * FColorDragDrop onto this widget. It "replaces" the FColorDragDrop widget with a preview SThemeColorBlock
+	 */
+	TSharedPtr<SThemeColorBlock> NewColorBlockPlaceholder;
+	/** The placeholder's color, also used to determine whether the placeholder is real or not (by IsValid) */
+	TSharedPtr<FLinearColor> NewColorPlaceholder;
+	/** Current x offset of the placeholder block */
+	float PlaceholderBlockOffset;
+	/** The initial grab offset when grabbing the placeholder */
+	FVector2D PlaceholderInitialGrabOffset;
+
+	/** A help text widget which appears when there are no children in this panel */
+	TSharedPtr<SWidget> EmptyHintTextBlock;
+
+	/**
+	 * The Color Theme that this SThemeColorBlockBar is displaying.
+	 * This is a TAttribute so it can re-get the theme when it changes rather than rely on a delegate to refresh it
+	 */
+	TAttribute< TSharedPtr<FColorTheme> > ColorTheme;
+	
+	/** Event called when a color block is clicked */
+	FOnLinearColorValueChanged OnSelectColor;
+	
+	/** Callback to pass to the Color Theme. Holds a handle to this bar's Refresh method */
+	FSimpleDelegate RefreshCallback;
+
+	/** Callback to pass down to the FColorDragDrop for it to show the trash */
+	FSimpleDelegate ShowTrashCallback;
+	/** Callback to pass down to the FColorDragDrop for it to hide the trash */
+	FSimpleDelegate HideTrashCallback;
+
+	/** Whether to use display sRGB color */
+	TAttribute<bool> bUseSRGB;
+
+	/** Whether or not the color uses Alpha or not */
+	TAttribute<bool> bUseAlpha;
+};
+
+
+DECLARE_DELEGATE_OneParam( FOnCurrentThemeChanged, TSharedPtr<FColorTheme> )
+
+
+/**
+ * SColorThemeBars include a ThemeColorBlocksBar in addition to a label.
+ * Clicking on one will select it and set the currently used color theme to it
+ */
+class SColorThemeBar
+	: public SCompoundWidget
+{
+public:
+	SLATE_BEGIN_ARGS(SColorThemeBar)
+		: _ColorTheme()
+		, _OnCurrentThemeChanged()
+		, _ShowTrashCallback()
+		, _HideTrashCallback()
+		, _UseSRGB()
+		, _UseAlpha()
+	{ }
+
+		/** The color theme that this bar is displaying */
+		SLATE_ATTRIBUTE(TSharedPtr<FColorTheme>, ColorTheme)
+
+		/** Event to be called when the current theme changes */
+		SLATE_EVENT(FOnCurrentThemeChanged, OnCurrentThemeChanged)
+
+		/** Callback to pass down to the FColorDragDrop for it to show the trash */
+		SLATE_EVENT(FSimpleDelegate, ShowTrashCallback)
+		/** Callback to pass down to the FColorDragDrop for it to hide the trash */
+		SLATE_EVENT(FSimpleDelegate, HideTrashCallback)
+		
+		/** Whether to display sRGB color */
+		SLATE_ATTRIBUTE(bool, UseSRGB)
+
+		/** Whether the ability to pick the alpha value is enabled */
+		SLATE_ATTRIBUTE(bool, UseAlpha)
+
+	SLATE_END_ARGS()
+
+	void Construct(const FArguments& InArgs);
+	
+	/**
+	 * The system calls this method to notify the widget that a mouse button was pressed within it. This event is bubbled.
+	 *
+	 * @param MyGeometry The Geometry of the widget receiving the event
+	 * @param MouseEvent Information about the input event
+	 * @return Whether the event was handled along with possible requests for the system to take action.
+	 */
+	virtual FReply OnMouseButtonDown( const FGeometry& MyGeometry, const FPointerEvent& MouseEvent );
+	
+private:
+
+	FString GetThemeName() const;
+
+	/** Text Block which shows the Color Theme's name */
+	TSharedPtr<STextBlock> ThemeNameText;
+
+	/** Color Theme that this bar is displaying */
+	TWeakPtr<FColorTheme> ColorTheme;
+
+	/** Callback to execute when the global current theme has changed */
+	FOnCurrentThemeChanged OnCurrentThemeChanged;
+
+	/** Callback to pass down to the FColorDragDrop for it to show the trash */
+	FSimpleDelegate ShowTrashCallback;
+	/** Callback to pass down to the FColorDragDrop for it to hide the trash */
+	FSimpleDelegate HideTrashCallback;
+	
+	/** Whether to use display sRGB color */
+	TAttribute<bool> bUseSRGB;
+
+	/** Whether or not the color uses Alpha or not */
+	TAttribute<bool> bUseAlpha;
+};
+
+
+/**
+ * There should only ever be a single SColorThemesViewer. It is the widget which
+ * manages all color themes and displays in a nice list format.
+ */
+class SColorThemesViewer
+	: public SCompoundWidget
+{
+public:
+
+	SLATE_BEGIN_ARGS(SColorThemesViewer)
+		: _UseAlpha()
+	{ }
+
+		/** Whether the ability to pick the alpha value is enabled */
+		SLATE_ATTRIBUTE(bool, UseAlpha)
+
+	SLATE_END_ARGS()
+	
+	void Construct(const FArguments& InArgs);
+	
+	/** Gets the current color theme */
+	TSharedPtr<FColorTheme> GetCurrentColorTheme() const;
+
+	/** Sets the UseAlpha attribute */
+	void SetUseAlpha( const TAttribute<bool>& InUseAlpha );
+
+	/** Load the color theme settings from the config */
+	static void LoadColorThemesFromIni();
+	/** Save the color theme settings to the config */
+	static void SaveColorThemesToIni();
+	
+	/** Callbacks to execute whenever we change the global current theme */
+	DECLARE_EVENT( SColorThemesViewer, FCurrentThemeChangedEvent );
+	FCurrentThemeChangedEvent& OnCurrentThemeChanged() { return CurrentThemeChangedEvent; }
+
+	/** Whether to use display sRGB color */
+	static bool bSRGBEnabled;
+
+	void MenuToStandardNoReturn();
+
+private:
+
+	FReply NewColorTheme();
+	FReply DuplicateColorTheme();
+	FReply DeleteColorTheme();
+	FReply AcceptThemeName();
+	bool CanAcceptThemeName() const;
+	void ChangeThemeName(const FText& InText);
+	EVisibility OnGetErrorTextVisibility() const;
+
+	/** Sets the current color theme to the existing theme */
+	void SetCurrentColorTheme(TSharedPtr<FColorTheme> NewTheme);
+
+	FReply MenuToStandard();
+	FReply MenuToRename();
+	FReply MenuToDelete();
+	void MenuToTrash();
+	
+	bool OnReadUseSRGB() const;
+	bool OnReadUseAlpha() const;
+	
+	/** Refreshes the list, saves colors to ini, and sends the menu to standard */
+	void RefreshThemes();
+
+	TSharedRef<ITableRow> OnGenerateColorThemeBars(TSharedPtr<FColorTheme> InItem, const TSharedRef<STableViewBase>& OwnerTable);
+
+	/** Gets the default color theme, optionally creates it if not present */
+	static TSharedPtr<FColorTheme> GetDefaultColorTheme(bool bCreateNew = false);
+	/** Gets the color theme, creates it if not present */
+	static TSharedPtr<FColorTheme> GetColorTheme(const FString& ThemeName);
+	/** Checks to see if this is a color theme, returns success */
+	static TSharedPtr<FColorTheme> IsColorTheme(const FString& ThemeName);
+	/** Makes the passed theme name unique so it doesn't clash with pre-existing themes */
+	static FString MakeUniqueThemeName(const FString& ThemeName);
+	/** Creates a new theme, ensuring the name is unique */
+	static TSharedPtr<FColorTheme> NewColorTheme(const FString& ThemeName, const TArray< TSharedPtr<FLinearColor> >& ThemeColors = TArray< TSharedPtr<FLinearColor> >());
+	
+	/** A list of all the color themes for the list view to use */
+	TSharedPtr< SListView <TSharedPtr <FColorTheme> > > ColorThemeList;
+
+	/** The menu is a widget at the bottom of the SColorThemesViewer with variable content */
+	TSharedPtr<SBorder> Menu;
+	/** The menu alternates between a standard menu, a rename menu, a confirm delete menu, and a color trash */
+	TSharedPtr<SWidget> MenuStandard;
+	TSharedPtr<SWidget> MenuRename;
+	TSharedPtr<SWidget> MenuConfirmDelete;
+	TSharedPtr<SWidget> MenuTrashColor;
+	/** If any error occurs, it needs to be displayed here as SPopupErrorText won't work with nested combo buttons */
+	TSharedPtr<SErrorText> ErrorText;
+
+	/** The text box for renaming themes */
+	TSharedPtr<SEditableTextBox> RenameTextBox;
+
+	/** Callbacks to execute whenever we change the global current theme */
+	FCurrentThemeChangedEvent CurrentThemeChangedEvent;
+
+	/** Whether or not the color uses Alpha or not */
+	TAttribute<bool> bUseAlpha;
+
+	/** A static holder of the color themes for the entire program */
+	static TArray< TSharedPtr<FColorTheme> > ColorThemes;
+	/** A static pointer to the color theme that is currently selected for the entire program */
+	static TWeakPtr<FColorTheme> CurrentlySelectedThemePtr;
+};
+
+
+/**
+ * This operation is a color which can be dragged and dropped between widgets.
+ * Represents a SThemeColorBlock that is dragged around, and can be dropped into a color trash.
+ */
+class FColorDragDrop
+	: public FDragDropOperation
+{
+public:
+
+	DRAG_DROP_OPERATOR_TYPE(FColorDragDrop, FDragDropOperation)
+
+	/**
+	 * Invoked when the drag and drop operation has ended.
+	 * 
+	 * @param bDropWasHandled   true when the drop was handled by some widget; false otherwise
+	 * @param MouseEvent        The mouse event which caused the on drop to be called.
+	 */
+	virtual void OnDrop( bool bDropWasHandled, const FPointerEvent& MouseEvent ) override;
+
+	/** 
+	 * Called when the mouse was moved during a drag and drop operation
+	 *
+	 * @param DragDropEvent    The event that describes this drag drop operation.
+	 */
+	virtual void OnDragged( const class FDragDropEvent& DragDropEvent ) override;
+	
+	/** Gets the widget that will serve as the decorator unless overriden. If you do not override, you will have no decorator */
+	virtual TSharedPtr<SWidget> GetDefaultDecorator() const override;
+	
+	/**
+	 * Makes a new FColorDragDrop to hold on to
+	 * @param InColor			The color to be dragged and dropped
+	 * @param bSRGB				Whether the color is sRGB
+	 * @param bUseAlpha			Whether the colors alpha is important
+	 * @param TrashShowCallback	Called when this operation is created
+	 * @param TrashHideCallback Called when this operation is dropped
+	 * @param Origin			The SThemeColorBlockBar that this operation is from
+	 * @param OriginPosition	The position in it's origin that it is from
+	 */
+	static TSharedRef<FColorDragDrop> New( FLinearColor InColor, bool bSRGB, bool bUseAlpha,
+		FSimpleDelegate TrashShowCallback = FSimpleDelegate(), FSimpleDelegate TrashHideCallback = FSimpleDelegate(),
+		TSharedPtr<SThemeColorBlocksBar> Origin = TSharedPtr<SThemeColorBlocksBar>(), int32 OriginPosition = 0 );
+
+	/** The color currently held onto by this drag drop operation */
+	FLinearColor Color;
+
+	/** Whether or not the color uses SRGB or not */
+	bool bUseSRGB;
+
+	/** Whether or not the color uses Alpha or not */
+	bool bUseAlpha;
+
+	/**
+	 * The origin is the SThemeColorBlockBar where this operation is from.
+	 * Upon dropping, this operation jumps back to it's origin if possible
+	 */
+	TWeakPtr<SThemeColorBlocksBar> OriginBar;
+
+	/** The origin position */
+	int32 OriginBarPosition;
+
+	/** Callback to show the trash when this widget is created */
+	FSimpleDelegate ShowTrash;
+
+	/** Callback to hide the trash when this widget is dropped */
+	FSimpleDelegate HideTrash;
+
+	/** Flag which ensures that OnDrop will not replace this block in it's origin */
+	bool bSetForDeletion;
+
+	/** The size of the drag and drop color block */
+	FVector2D BlockSize;
+};
diff --git a/Engine/Source/Runtime/AppFramework/Public/Widgets/Colors/SComplexGradient.h b/Engine/Source/Runtime/AppFramework/Public/Widgets/Colors/SComplexGradient.h
new file mode 100644
--- /dev/null
+++ b/Engine/Source/Runtime/AppFramework/Public/Widgets/Colors/SComplexGradient.h
@@ -0,0 +1,54 @@
+// Copyright 1998-2014 Epic Games, Inc. All Rights Reserved.
+
+#pragma once
+
+
+/**
+ * Implements a Slate widget that renders a color gradient consisting of multiple stops.
+ */
+class APPFRAMEWORK_API SComplexGradient
+	: public SCompoundWidget
+{
+public:
+
+	SLATE_BEGIN_ARGS(SComplexGradient)
+		: _GradientColors()
+		, _HasAlphaBackground(false)
+		, _Orientation(Orient_Vertical)
+	{ }
+
+		/** The colors used in the gradient */
+		SLATE_ATTRIBUTE(TArray<FLinearColor>, GradientColors)
+
+		/** Whether a checker background is displayed for alpha viewing */
+		SLATE_ATTRIBUTE(bool, HasAlphaBackground)
+		
+		/** Horizontal or vertical gradient */
+		SLATE_ATTRIBUTE(EOrientation, Orientation)
+
+	SLATE_END_ARGS()
+
+	/**
+	 * Construct this widget
+	 *
+	 * @param InArgs The declaration data for this widget.
+	 */
+	void Construct( const FArguments& InArgs );
+
+protected:
+
+	// SCompoundWidget overrides
+
+	virtual int32 OnPaint( const FPaintArgs& Args, const FGeometry& AllottedGeometry, const FSlateRect& MyClippingRect, FSlateWindowElementList& OutDrawElements, int32 LayerId, const FWidgetStyle& InWidgetStyle, bool bParentEnabled ) const override;
+	
+private:
+
+	/** The colors used in the gradient. */
+	TAttribute<TArray<FLinearColor>> GradientColors;
+
+	/** Whether a checker background is displayed for alpha viewing. */
+	bool bHasAlphaBackground;
+	
+	/** Horizontal or vertical gradient. */
+	EOrientation Orientation;
+};
diff --git a/Engine/Source/Runtime/AppFramework/Public/Widgets/Colors/SSimpleGradient.h b/Engine/Source/Runtime/AppFramework/Public/Widgets/Colors/SSimpleGradient.h
new file mode 100644
--- /dev/null
+++ b/Engine/Source/Runtime/AppFramework/Public/Widgets/Colors/SSimpleGradient.h
@@ -0,0 +1,68 @@
+// Copyright 1998-2014 Epic Games, Inc. All Rights Reserved.
+
+#pragma once
+
+
+/**
+ * Implements a Slate widget that renders a simple color gradient.
+ */
+class APPFRAMEWORK_API SSimpleGradient
+	: public SCompoundWidget
+{
+public:
+
+		SLATE_BEGIN_ARGS(SSimpleGradient)
+			: _StartColor(FLinearColor(0.0f, 0.0f, 0.0f))
+			, _EndColor(FLinearColor(1.0f, 1.0f, 1.0f))
+			, _HasAlphaBackground(false)
+			, _Orientation(Orient_Vertical)
+			, _UseSRGB(true)
+		{ }
+
+		/** The leftmost gradient color */
+		SLATE_ATTRIBUTE(FLinearColor, StartColor)
+		
+		/** The rightmost gradient color */
+		SLATE_ATTRIBUTE(FLinearColor, EndColor)
+
+		/** Whether a checker background is displayed for alpha viewing */
+		SLATE_ATTRIBUTE(bool, HasAlphaBackground)
+
+		/** Horizontal or vertical gradient */
+		SLATE_ATTRIBUTE(EOrientation, Orientation)
+
+		/** Whether to display sRGB color */
+		SLATE_ATTRIBUTE(bool, UseSRGB)
+
+	SLATE_END_ARGS()
+
+	/**
+	 * Constructs the widget
+	 *
+	 * @param InArgs The declaration data for this widget
+	 */
+	void Construct( const FArguments& InArgs );
+
+protected:
+
+	// SCompoundWidget overrides
+
+	virtual int32 OnPaint( const FPaintArgs& Args, const FGeometry& AllottedGeometry, const FSlateRect& MyClippingRect, FSlateWindowElementList& OutDrawElements, int32 LayerId, const FWidgetStyle& InWidgetStyle, bool bParentEnabled ) const override;
+
+private:
+
+	/** The leftmost gradient color */
+	TAttribute<FLinearColor> StartColor;
+
+	/** The rightmost gradient color */
+	TAttribute<FLinearColor> EndColor;
+
+	/** Whether a checker background is displayed for alpha viewing */
+	bool bHasAlphaBackground;
+	
+	/** Horizontal or vertical gradient */
+	EOrientation Orientation;
+
+	/** Whether to display sRGB color */
+	TAttribute<bool> UseSRGB;
+};
diff --git a/Engine/Source/Runtime/AppFramework/Public/Widgets/Testing/STestSuite.h b/Engine/Source/Runtime/AppFramework/Public/Widgets/Testing/STestSuite.h
new file mode 100644
--- /dev/null
+++ b/Engine/Source/Runtime/AppFramework/Public/Widgets/Testing/STestSuite.h
@@ -0,0 +1,8 @@
+// Copyright 1998-2014 Epic Games, Inc. All Rights Reserved.
+
+#pragma once
+
+
+APPFRAMEWORK_API void RestoreSlateTestSuite();
+
+APPFRAMEWORK_API void MakeSplitterTest();
\ No newline at end of file
diff --git a/Engine/Source/Runtime/Engine/Private/UnrealEngine.cpp b/Engine/Source/Runtime/Engine/Private/UnrealEngine.cpp
--- a/Engine/Source/Runtime/Engine/Private/UnrealEngine.cpp
+++ b/Engine/Source/Runtime/Engine/Private/UnrealEngine.cpp
@@ -83,6 +83,8 @@
 
 #include "Engine/GameInstance.h"
 #include "HotReloadInterface.h"
+#include "STestSuite.h"
+
 
 DEFINE_LOG_CATEGORY_STATIC(LogEngine, Log, All);
 
diff --git a/Engine/Source/Runtime/Slate/Public/Slate.h b/Engine/Source/Runtime/Slate/Public/Slate.h
--- a/Engine/Source/Runtime/Slate/Public/Slate.h
+++ b/Engine/Source/Runtime/Slate/Public/Slate.h
@@ -129,8 +129,6 @@ DECLARE_CYCLE_STAT_EXTERN( TEXT("Slate Misc Time"), STAT_SlateMiscTime, STATGROU
 #include "SVolumeControl.h"
 #include "SColorSpectrum.h"
 #include "SColorWheel.h"
-#include "SColorThemes.h"
-#include "SColorPicker.h"
 ///
 #include "SToolBarButtonBlock.h"
 #include "SToolBarComboButtonBlock.h"
@@ -172,6 +170,3 @@ DECLARE_CYCLE_STAT_EXTERN( TEXT("Slate Misc Time"), STAT_SlateMiscTime, STATGROU
 // Old docking
 #include "SDockableTab.h"
 #include "SDockTabStack.h"
-
-// Test suite
-#include "STestSuite.h"
diff --git a/Engine/Source/Runtime/Slate/Public/Widgets/Colors/SColorSpectrum.h b/Engine/Source/Runtime/Slate/Public/Widgets/Colors/SColorSpectrum.h
--- a/Engine/Source/Runtime/Slate/Public/Widgets/Colors/SColorSpectrum.h
+++ b/Engine/Source/Runtime/Slate/Public/Widgets/Colors/SColorSpectrum.h
@@ -6,7 +6,7 @@
 /**
  * Implements the color spectrum widget.
  */
-class SColorSpectrum
+class SLATE_API SColorSpectrum
 	: public SLeafWidget
 {
 public:
diff --git a/Engine/Source/Runtime/Slate/Public/Widgets/Colors/SColorWheel.h b/Engine/Source/Runtime/Slate/Public/Widgets/Colors/SColorWheel.h
--- a/Engine/Source/Runtime/Slate/Public/Widgets/Colors/SColorWheel.h
+++ b/Engine/Source/Runtime/Slate/Public/Widgets/Colors/SColorWheel.h
@@ -6,7 +6,7 @@
 /**
  * Implements the color wheel widget.
  */
-class SColorWheel
+class SLATE_API SColorWheel
 	: public SLeafWidget
 {
 public:
