Slate: Moved SWizard into AppFramework module

diff --git a/Engine/Source/Developer/SessionLauncher/Private/Widgets/SSessionLauncher.h b/Engine/Source/Developer/SessionLauncher/Private/Widgets/SSessionLauncher.h
--- a/Engine/Source/Developer/SessionLauncher/Private/Widgets/SSessionLauncher.h
+++ b/Engine/Source/Developer/SessionLauncher/Private/Widgets/SSessionLauncher.h
@@ -123,7 +123,4 @@ private:
 
 	// Holds the widget switcher.
 	TSharedPtr<SWidgetSwitcher> WidgetSwitcher;
-
-	// Holds the wizard widget.
-	TSharedPtr<SWizard> Wizard;
 };
diff --git a/Engine/Source/Editor/GameProjectGeneration/GameProjectGeneration.Build.cs b/Engine/Source/Editor/GameProjectGeneration/GameProjectGeneration.Build.cs
--- a/Engine/Source/Editor/GameProjectGeneration/GameProjectGeneration.Build.cs
+++ b/Engine/Source/Editor/GameProjectGeneration/GameProjectGeneration.Build.cs
@@ -17,6 +17,7 @@ public class GameProjectGeneration : ModuleRules
 		PrivateDependencyModuleNames.AddRange(
 			new string[] {
 				"Analytics",
+                "AppFramework",
 				"Core",
 				"CoreUObject",
 				"Engine",
diff --git a/Engine/Source/Editor/GameProjectGeneration/Private/SNewClassDialog.cpp b/Engine/Source/Editor/GameProjectGeneration/Private/SNewClassDialog.cpp
--- a/Engine/Source/Editor/GameProjectGeneration/Private/SNewClassDialog.cpp
+++ b/Engine/Source/Editor/GameProjectGeneration/Private/SNewClassDialog.cpp
@@ -9,9 +9,12 @@
 #include "DesktopPlatformModule.h"
 #include "Editor/Documentation/Public/IDocumentation.h"
 #include "EditorClassUtils.h"
+#include "SWizard.h"
+
 
 #define LOCTEXT_NAMESPACE "GameProjectGeneration"
 
+
 struct FParentClassItem
 {
 	GameProjectUtils::FNewClassInfo ParentClassInfo;
diff --git a/Engine/Source/Editor/GameProjectGeneration/Private/SNewProjectWizard.cpp b/Engine/Source/Editor/GameProjectGeneration/Private/SNewProjectWizard.cpp
--- a/Engine/Source/Editor/GameProjectGeneration/Private/SNewProjectWizard.cpp
+++ b/Engine/Source/Editor/GameProjectGeneration/Private/SNewProjectWizard.cpp
@@ -8,12 +8,15 @@
 #include "SScrollBorder.h"
 #include "TemplateCategory.h"
 #include "GameProjectGenerationModule.h"
+#include "SWizard.h"
+
 
 #define LOCTEXT_NAMESPACE "NewProjectWizard"
 
 FName SNewProjectWizard::TemplatePageName = TEXT("Template");
 FName SNewProjectWizard::NameAndLocationPageName = TEXT("NameAndLocation");
 
+
 namespace
 {
 	const float ThumbnailSize = 64.f, ThumbnailPadding = 5.f;
diff --git a/Engine/Source/Editor/GameProjectGeneration/Private/SNewProjectWizard.h b/Engine/Source/Editor/GameProjectGeneration/Private/SNewProjectWizard.h
--- a/Engine/Source/Editor/GameProjectGeneration/Private/SNewProjectWizard.h
+++ b/Engine/Source/Editor/GameProjectGeneration/Private/SNewProjectWizard.h
@@ -177,7 +177,7 @@ private:
 private:
 
 	/** The wizard widget */
-	TSharedPtr<SWizard> MainWizard;
+	TSharedPtr<class SWizard> MainWizard;
 
 	FString LastBrowsePath;
 	FString CurrentProjectFileName;
diff --git a/Engine/Source/Editor/IntroTutorials/Private/SEditorTutorialImporter.cpp b/Engine/Source/Editor/IntroTutorials/Private/SEditorTutorialImporter.cpp
--- a/Engine/Source/Editor/IntroTutorials/Private/SEditorTutorialImporter.cpp
+++ b/Engine/Source/Editor/IntroTutorials/Private/SEditorTutorialImporter.cpp
@@ -6,9 +6,12 @@
 #include "IDocumentationPage.h"
 #include "EditorTutorial.h"
 #include "DesktopPlatformModule.h"
+#include "SWizard.h"
+
 
 #define LOCTEXT_NAMESPACE "SEditorTutorialImporter"
 
+
 void SEditorTutorialImporter::Construct(const FArguments& InArgs)
 {
 	ParentWindow = InArgs._ParentWindow;
diff --git a/Engine/Source/Runtime/AppFramework/Private/Widgets/Workflow/SWizard.cpp b/Engine/Source/Runtime/AppFramework/Private/Widgets/Workflow/SWizard.cpp
new file mode 100644
--- /dev/null
+++ b/Engine/Source/Runtime/AppFramework/Private/Widgets/Workflow/SWizard.cpp
@@ -0,0 +1,352 @@
+// Copyright 1998-2014 Epic Games, Inc. All Rights Reserved.
+
+#include "AppFrameworkPrivatePCH.h"
+#include "SWizard.h"
+
+
+#define LOCTEXT_NAMESPACE "SWizard"
+
+
+/* SWizard interface
+ *****************************************************************************/
+
+bool SWizard::CanShowPage( int32 PageIndex ) const
+{
+	if (Pages.IsValidIndex(PageIndex))
+	{
+		return Pages[PageIndex].CanShow();
+	}
+
+	return false;
+}
+
+
+void SWizard::Construct( const FArguments& InArgs )
+{
+	DesiredSize = InArgs._DesiredSize.Get();
+	OnCanceled = InArgs._OnCanceled;
+	OnFinished = InArgs._OnFinished;
+	OnFirstPageBackClicked = InArgs._OnFirstPageBackClicked;
+
+	TSharedPtr<SVerticalBox> PageListBox;
+	TSharedPtr< SUniformGridPanel > ButtonGrid; 
+
+	ChildSlot
+	[
+		SNew(SVerticalBox)
+
+		+ SVerticalBox::Slot()
+			.FillHeight(1.0)			
+			[
+				SNew(SHorizontalBox)
+
+				+ SHorizontalBox::Slot()
+					.AutoWidth()
+					.Padding(0.0f, 0.0f, 20.0f, 0.0f)
+					[
+						SAssignNew(PageListBox, SVerticalBox)
+							.Visibility(InArgs._ShowPageList ? EVisibility::Visible : EVisibility::Collapsed)
+					]
+
+				+ SHorizontalBox::Slot()
+					.FillWidth(1.0f)
+					[
+						// widget switcher
+						SAssignNew(WidgetSwitcher, SWidgetSwitcher)
+					]
+			]
+
+		+ SVerticalBox::Slot()
+			.AutoHeight()
+			.HAlign(HAlign_Right)
+			.Padding(0.0f, 10.0f, 0.0f, 0.0f)
+			[
+				SAssignNew(ButtonGrid, SUniformGridPanel)
+					.SlotPadding(FCoreStyle::Get().GetMargin("StandardDialog.SlotPadding"))
+					.MinDesiredSlotWidth(FCoreStyle::Get().GetFloat("StandardDialog.MinDesiredSlotWidth"))
+					.MinDesiredSlotHeight(FCoreStyle::Get().GetFloat("StandardDialog.MinDesiredSlotHeight"))
+
+				+ SUniformGridPanel::Slot(0, 0)
+					[
+						// 'Prev' button
+						SNew(SButton)
+							.HAlign(HAlign_Center)
+							.VAlign(VAlign_Center)
+							.ContentPadding(FCoreStyle::Get().GetMargin("StandardDialog.ContentPadding"))
+							.IsEnabled(this, &SWizard::HandlePrevButtonIsEnabled)
+							.OnClicked(this, &SWizard::HandlePrevButtonClicked)
+							.Visibility(this, &SWizard::HandlePrevButtonVisibility)
+							.ToolTipText(LOCTEXT("PrevButtonTooltip", "Go back to the previous step"))
+							[
+								SNew(SHorizontalBox)
+
+								+ SHorizontalBox::Slot()
+									.Padding(2.0f, 0.0f)
+									.AutoWidth()
+									.VAlign(VAlign_Center)
+									[
+										SNew(SImage)
+											.Image(FCoreStyle::Get().GetBrush("Wizard.BackIcon"))
+											.ColorAndOpacity(FLinearColor(0.05f, 0.05f, 0.05f))
+									]
+
+								+ SHorizontalBox::Slot()
+									.AutoWidth()
+									.VAlign(VAlign_Center)
+									[
+										SNew(STextBlock)
+											.Text(LOCTEXT("PrevButtonLabel", "Back"))
+									]
+							]
+					]
+
+				+ SUniformGridPanel::Slot(1, 0)
+				[
+					// 'Next' button
+					SNew(SButton)
+					.HAlign(HAlign_Center)
+					.VAlign(VAlign_Center)
+					.ContentPadding(FCoreStyle::Get().GetMargin("StandardDialog.ContentPadding"))
+					.IsEnabled(this, &SWizard::HandleNextButtonIsEnabled)
+					.OnClicked(this, &SWizard::HandleNextButtonClicked)
+					.Visibility(this, &SWizard::HandleNextButtonVisibility)
+					.ToolTipText(LOCTEXT("NextButtonTooltip", "Go to the next step"))
+					[
+						SNew(SHorizontalBox)
+
+						+SHorizontalBox::Slot()
+						.AutoWidth()
+						.VAlign(VAlign_Center)
+						[
+							SNew(STextBlock)
+							.Text(LOCTEXT("NextButtonLabel", "Next"))
+						]
+
+						+SHorizontalBox::Slot()
+						.Padding(2.0f, 0.0f)
+						.AutoWidth()
+						.VAlign(VAlign_Center)
+						[
+							SNew(SImage)
+							.Image(FCoreStyle::Get().GetBrush("Wizard.NextIcon"))
+							.ColorAndOpacity(FLinearColor(0.05f, 0.05f, 0.05f))
+						]
+					]
+				]
+
+				+ SUniformGridPanel::Slot(2,0)
+					[
+						// 'Finish' button
+						SNew(SButton)
+							.HAlign(HAlign_Center)
+							.VAlign(VAlign_Center)
+							.ContentPadding(FCoreStyle::Get().GetMargin("StandardDialog.ContentPadding"))
+							.IsEnabled(InArgs._CanFinish)
+							.OnClicked(this, &SWizard::HandleFinishButtonClicked)
+							.ToolTipText(InArgs._FinishButtonToolTip)
+							.Text(InArgs._FinishButtonText)
+					]
+			]
+	];
+
+	if (InArgs._ShowCancelButton)
+	{
+		ButtonGrid->AddSlot(3, 0)
+		[
+			// cancel button
+			SNew(SButton)
+				.HAlign(HAlign_Center)
+				.VAlign(VAlign_Center)
+				.ContentPadding(FCoreStyle::Get().GetMargin("StandardDialog.ContentPadding"))
+				.OnClicked(this, &SWizard::HandleCancelButtonClicked)
+				.ToolTipText(LOCTEXT("CancelButtonTooltip", "Cancel this wizard"))
+				.Text(LOCTEXT("CancelButtonLabel", "Cancel"))
+		];
+	}
+
+	// populate wizard with pages
+	for (int32 SlotIndex = 0; SlotIndex < InArgs.Slots.Num(); ++SlotIndex)
+	{
+		FWizardPage* Page = InArgs.Slots[SlotIndex];
+
+		Pages.Add(Page);
+
+		if (InArgs._ShowPageList)
+		{
+			PageListBox->AddSlot()
+				.AutoHeight()
+				[
+					SNew(SCheckBox)
+						.IsChecked(this, &SWizard::HandlePageButtonIsChecked, SlotIndex)
+						.IsEnabled(this, &SWizard::HandlePageButtonIsEnabled, SlotIndex)
+						.OnCheckStateChanged(this, &SWizard::HandlePageButtonCheckStateChanged, SlotIndex)
+						.Padding(FMargin(8.0f, 4.0f, 24.0f, 4.0f))
+						.Style(&FCoreStyle::Get().GetWidgetStyle< FCheckBoxStyle >("ToggleButtonCheckbox"))
+						[
+							Page->GetButtonContent()
+						]
+				];
+		}
+
+		WidgetSwitcher->AddSlot()
+			[
+				Page->GetPageContent()
+			];
+	}
+
+	WidgetSwitcher->SetActiveWidgetIndex(InArgs._InitialPageIndex.Get());
+}
+
+
+void SWizard::ShowPage( int32 PageIndex )
+{
+	int32 ActivePageIndex = WidgetSwitcher->GetActiveWidgetIndex();
+
+	if (Pages.IsValidIndex(ActivePageIndex))
+	{
+		Pages[ActivePageIndex].OnLeave().ExecuteIfBound();
+	}
+
+	if (Pages.IsValidIndex(PageIndex) && Pages[PageIndex].CanShow())
+	{
+		WidgetSwitcher->SetActiveWidgetIndex(PageIndex);
+
+		// show the desired page
+		Pages[PageIndex].OnEnter().ExecuteIfBound();
+	}
+	else if (Pages.IsValidIndex(0) && Pages[0].CanShow())
+	{
+		WidgetSwitcher->SetActiveWidgetIndex(0);
+
+		// attempt to fall back to first page
+		Pages[0].OnEnter().ExecuteIfBound();
+	}
+	else
+	{
+		WidgetSwitcher->SetActiveWidgetIndex(INDEX_NONE);
+	}
+}
+
+
+/* SCompoundWidget overrides
+ *****************************************************************************/
+
+FVector2D SWizard::ComputeDesiredSize( ) const
+{
+	if (DesiredSize.IsZero())
+	{
+		return SCompoundWidget::ComputeDesiredSize();
+	}
+
+	return DesiredSize;
+}
+
+
+/* SWizard callbacks
+ *****************************************************************************/
+
+FReply SWizard::HandleCancelButtonClicked( )
+{
+	OnCanceled.ExecuteIfBound();
+
+	return FReply::Handled();
+}
+
+
+FReply SWizard::HandleFinishButtonClicked( )
+{
+	OnFinished.ExecuteIfBound();
+
+	return FReply::Handled();
+}
+
+
+FReply SWizard::HandleNextButtonClicked( )
+{
+	ShowPage(WidgetSwitcher->GetActiveWidgetIndex() + 1);
+
+	return FReply::Handled();
+}
+
+
+bool SWizard::HandleNextButtonIsEnabled( ) const
+{
+	return CanShowPage(WidgetSwitcher->GetActiveWidgetIndex() + 1);
+}
+
+
+EVisibility SWizard::HandleNextButtonVisibility( ) const
+{
+	if (Pages.IsValidIndex(WidgetSwitcher->GetActiveWidgetIndex() + 1))
+	{
+		return EVisibility::Visible;
+	}
+
+	return EVisibility::Hidden;
+}
+
+
+void SWizard::HandlePageButtonCheckStateChanged( ESlateCheckBoxState::Type NewState, int32 PageIndex )
+{
+	if (NewState == ESlateCheckBoxState::Checked)
+	{
+		ShowPage(PageIndex);
+	}
+}
+
+
+ESlateCheckBoxState::Type SWizard::HandlePageButtonIsChecked( int32 PageIndex ) const
+{
+	if (PageIndex == WidgetSwitcher->GetActiveWidgetIndex())
+	{
+		return ESlateCheckBoxState::Checked;
+	}
+
+	return ESlateCheckBoxState::Unchecked;
+}
+
+
+bool SWizard::HandlePageButtonIsEnabled( int32 PageIndex ) const
+{
+	return CanShowPage(PageIndex);
+}
+
+
+FReply SWizard::HandlePrevButtonClicked( )
+{
+	if ( WidgetSwitcher->GetActiveWidgetIndex() == 0 && OnFirstPageBackClicked.IsBound() )
+	{
+		return OnFirstPageBackClicked.Execute();
+	}
+	else
+	{
+		ShowPage(WidgetSwitcher->GetActiveWidgetIndex() - 1);
+	}
+
+	return FReply::Handled();
+}
+
+
+bool SWizard::HandlePrevButtonIsEnabled( ) const
+{
+	if ( WidgetSwitcher->GetActiveWidgetIndex() == 0 && OnFirstPageBackClicked.IsBound() )
+	{
+		return true;
+	}
+
+	return CanShowPage(WidgetSwitcher->GetActiveWidgetIndex() - 1);
+}
+
+
+EVisibility SWizard::HandlePrevButtonVisibility( ) const
+{
+	if (WidgetSwitcher->GetActiveWidgetIndex() > 0 || OnFirstPageBackClicked.IsBound())
+	{
+		return EVisibility::Visible;
+	}
+
+	return EVisibility::Hidden;
+}
+
+
+#undef LOCTEXT_NAMESPACE
diff --git a/Engine/Source/Runtime/AppFramework/Public/Widgets/Workflow/SWizard.h b/Engine/Source/Runtime/AppFramework/Public/Widgets/Workflow/SWizard.h
new file mode 100644
--- /dev/null
+++ b/Engine/Source/Runtime/AppFramework/Public/Widgets/Workflow/SWizard.h
@@ -0,0 +1,316 @@
+// Copyright 1998-2014 Epic Games, Inc. All Rights Reserved.
+
+#pragma once
+
+
+/**
+ * Delegate type for getting a page's enabled state.
+ *
+ * The first parameter is the page index.
+ */
+DECLARE_DELEGATE_OneParam(FOnWizardPageIsEnabled, int32)
+
+
+/**
+ * Implements a wizard widget.
+ */
+class APPFRAMEWORK_API SWizard
+	: public SCompoundWidget
+{
+public:
+
+	/**
+	 * Implements a wizard page.
+	 */
+	class APPFRAMEWORK_API FWizardPage
+	{
+	public:
+
+		HACK_SLATE_SLOT_ARGS(FWizardPage)
+			: _ButtonContent()
+			, _CanShow(true)
+			, _PageContent()
+		{ }
+
+			/** Holds the content of the button to be shown in the wizard's page selector. */
+			SLATE_NAMED_SLOT(FArguments, ButtonContent)
+
+			/** Holds a flag indicating whether this page can be activated (default = true). */
+			SLATE_ATTRIBUTE(bool, CanShow)
+
+			/** Exposes a delegate to be invoked when this page is being activated. */
+			SLATE_EVENT(FSimpleDelegate, OnEnter)
+
+			/** Exposes a delegate to be invoked when this page is being deactivated. */
+			SLATE_EVENT(FSimpleDelegate, OnLeave)
+
+			/** Holds the page content. */
+			SLATE_DEFAULT_SLOT(FArguments, PageContent)
+
+		SLATE_END_ARGS()
+
+	public:
+
+		/**
+		 * Creates and initializes a new instance.
+		 *
+		 * @param InArgs The arguments for the page.
+		 */
+		FWizardPage( const FArguments& InArgs )
+			: ButtonContent(InArgs._ButtonContent)
+			, Showable(InArgs._CanShow)
+			, OnEnterDelegate(InArgs._OnEnter)
+			, OnLeaveDelegate(InArgs._OnLeave)
+			, PageContent(InArgs._PageContent)
+		{ }
+
+	public:
+
+		/**
+		 * Checks whether the page can be shown.
+		 *
+		 * @return true if the page can be shown, false otherwise.
+		 */
+		bool CanShow( ) const
+		{
+			return Showable.Get();
+		}
+
+		/**
+		 * Gets the button content.
+		 *
+		 * @return The button content widget.
+		 */
+		const TSharedRef<SWidget>& GetButtonContent( ) const
+		{
+			return ButtonContent.Widget;
+		}
+
+		/**
+		 * Gets the page content.
+		 *
+		 * @return The page contents.
+		 */
+		const TSharedRef<SWidget>& GetPageContent( ) const
+		{
+			return PageContent.Widget;
+		}
+
+		/**
+		 * Gets a delegate to be invoked when this page is being entered.
+		 *
+		 * @return The delegate.
+		 */
+		FSimpleDelegate& OnEnter( )
+		{
+			return OnEnterDelegate;
+		}
+
+		/**
+		 * Gets a delegate to be invoked when this page is being left.
+		 *
+		 * @return The delegate.
+		 */
+		FSimpleDelegate& OnLeave( )
+		{
+			return OnLeaveDelegate;
+		}
+
+	private:
+
+		// Holds the button content.
+		TAlwaysValidWidget ButtonContent;
+
+		// Holds a flag indicating whether this page can be activated.
+		TAttribute<bool> Showable;
+
+		// Holds a delegate to be invoked when the page is activated.
+		FSimpleDelegate OnEnterDelegate;
+
+		// Holds a delegate to be invoked when the page is deactivated.
+		FSimpleDelegate OnLeaveDelegate;
+
+		// Holds the page content.
+		TAlwaysValidWidget PageContent;
+	};
+
+public:
+
+	SLATE_BEGIN_ARGS(SWizard)
+		: _CanFinish(true)
+		, _FinishButtonText(NSLOCTEXT("SWizard", "DefaultFinishButtonText", "Finish").ToString())
+		, _FinishButtonToolTip(NSLOCTEXT("SWizard", "DefaultFinishButtonTooltip", "Finish the wizard").ToString())
+		, _InitialPageIndex(0)
+		, _DesiredSize(FVector2D(0, 0))
+		, _ShowPageList(true)
+		, _ShowCancelButton(true)
+	{ }
+
+		SLATE_SUPPORTS_SLOT_WITH_ARGS(FWizardPage)
+
+		/** Holds a flag indicating whether the 'Finish' button is enabled. */
+		SLATE_ATTRIBUTE(bool, CanFinish)
+
+		/** Holds the label text of the wizard's 'Finish' button. */
+		SLATE_TEXT_ATTRIBUTE(FinishButtonText)
+
+		/** Holds the tool tip text of the wizard's 'Finish' button. */
+		SLATE_TEXT_ATTRIBUTE(FinishButtonToolTip)
+
+		/** Holds the index of the initial page to be displayed (0 = default). */
+		SLATE_ATTRIBUTE(int32, InitialPageIndex)
+
+		/** Holds the desired size. */
+		SLATE_ATTRIBUTE(FVector2D, DesiredSize)
+
+		/** Exposes a delegate to be invoked when the wizard is canceled. */
+		SLATE_EVENT(FSimpleDelegate, OnCanceled)
+
+		/** Exposes a delegate to be invoked when the wizard is finished. */
+		SLATE_EVENT(FSimpleDelegate, OnFinished)
+
+		/** Exposes a delegate to be invoked when the wizard's 'Back' button is clicked. This button is only visible if this delegate is set and only on the first page. */
+		SLATE_EVENT(FOnClicked, OnFirstPageBackClicked)
+
+		/** Exposes a delegate to be invoked when the wizard's 'Next' button is clicked. */
+		SLATE_EVENT(FOnClicked, OnNextClicked)
+
+		/** Exposes a delegate to be invoked when the wizard's 'Previous' button is clicked. */
+		SLATE_EVENT(FOnClicked, OnPrevClicked)
+
+		/** Holds a flag indicating whether the page list should be shown (default = true). */
+		SLATE_ARGUMENT(bool, ShowPageList)
+
+		/** Holds a flag indicating whether the cancel button should be shown (default = true). */
+		SLATE_ARGUMENT(bool, ShowCancelButton)
+
+	SLATE_END_ARGS()
+
+public:
+
+	/**
+	 * Checks whether the page with the specified index can be shown.
+	 *
+	 * @param PageIndex The index of the page to check.
+	 * @return true if the page can be shown, false otherwise.
+	 */
+	bool CanShowPage( int32 PageIndex ) const;
+
+	/**
+	 * Construct this widget.
+	 *
+	 * @param InArgs The declaration data for this widget.
+	 */
+	void Construct( const FArguments& InArgs );
+
+	/**
+	 * Gets the number of pages that this wizard contains.
+	 *
+	 * @return The number of pages.
+	 */
+	int32 GetNumPages( ) const
+	{
+		return WidgetSwitcher->GetNumWidgets();
+	}
+
+	/**
+	 * Gets the index of the specified wizard page widget.
+	 *
+	 * @param PageWidget The page widget to get the index for.
+	 * @return The index of the page, or INDEX_NONE if not found.
+	 */
+	int32 GetPageIndex( const TSharedRef<SWidget>& PageWidget ) const
+	{
+		return WidgetSwitcher->GetWidgetIndex(PageWidget);
+	}
+
+	/**
+	 * Attempts to show the page with the specified index.
+	 *
+	 * @param PageIndex The index of the page to show.
+	 */
+	void ShowPage( int32 PageIndex );
+
+public:
+
+	// SCompoundWidget interface
+
+	virtual FVector2D ComputeDesiredSize( ) const override;
+
+public:
+
+	/**
+	 * Returns a new slot for a page.
+	 *
+	 * @param PageName The name of the page.
+	 * @return A new slot.
+	 */
+	static FWizardPage::FArguments Page( )
+	{
+		FWizardPage::FArguments Args;
+		return Args;
+	}
+
+private:
+
+	// Callback for clicking the 'Cancel' button.
+	FReply HandleCancelButtonClicked( );
+
+	// Callback for clicking the 'Launch' button.
+	FReply HandleFinishButtonClicked( );
+
+	// Callback for clicking the 'Next' button.
+	FReply HandleNextButtonClicked( );
+
+	// Callback for getting the enabled state of the 'Next' button.
+	bool HandleNextButtonIsEnabled( ) const;
+
+	// Callback for getting the visibility of the 'Next' button.
+	EVisibility HandleNextButtonVisibility( ) const;
+
+	// Callback for getting the checked state of a page button.
+	void HandlePageButtonCheckStateChanged( ESlateCheckBoxState::Type NewState, int32 PageIndex );
+
+	// Callback for clicking a page button.
+	ESlateCheckBoxState::Type HandlePageButtonIsChecked( int32 PageIndex ) const;
+
+	// Callback for getting the enabled state of a page button.
+	bool HandlePageButtonIsEnabled( int32 PageIndex ) const;
+
+	// Callback for clicking the 'Previous' button.
+	FReply HandlePrevButtonClicked( );
+
+	// Callback for getting the enabled state of the 'Previous' button.
+	bool HandlePrevButtonIsEnabled( ) const;
+
+	// Callback for getting the visibility of the 'Previous' button.
+	EVisibility HandlePrevButtonVisibility( ) const;
+
+private:
+
+	// Holds the wizard's desired size.
+	FVector2D DesiredSize;
+
+	// Holds the collection of wizard pages.
+	TIndirectArray<FWizardPage> Pages;
+
+	// Holds the widget switcher.
+	TSharedPtr<SWidgetSwitcher> WidgetSwitcher;
+
+private:
+
+	// Holds a delegate to be invoked when the 'Cancel' button has been clicked.
+	FSimpleDelegate OnCanceled;
+
+	// Holds a delegate to be invoked when the 'Finish' button has been clicked.
+	FSimpleDelegate OnFinished;
+
+	// Holds a delegate to be invoked when the 'Next' button has been clicked.
+	FOnClicked OnNextClicked;
+
+	// Holds a delegate to be invoked when the 'Previous' button has been clicked.
+	FOnClicked OnPrevClicked;
+
+	// Holds a delegate to be invoked when the 'Previous' button has been clicked on the first page.
+	FOnClicked OnFirstPageBackClicked;
+};
diff --git a/Engine/Source/Runtime/Slate/Private/Widgets/Navigation/SWizard.cpp b/Engine/Source/Runtime/Slate/Private/Widgets/Navigation/SWizard.cpp
deleted file mode 100644
diff --git a/Engine/Source/Runtime/Slate/Public/Slate.h b/Engine/Source/Runtime/Slate/Public/Slate.h
--- a/Engine/Source/Runtime/Slate/Public/Slate.h
+++ b/Engine/Source/Runtime/Slate/Public/Slate.h
@@ -146,7 +146,6 @@ DECLARE_CYCLE_STAT_EXTERN( TEXT("Slate Misc Time"), STAT_SlateMiscTime, STATGROU
 #include "SWidgetSwitcher.h"
 #include "SSuggestionTextBox.h"
 #include "SBreadcrumbTrail.h"
-#include "SWizard.h"
 #include "STextComboBox.h"
 #include "SNumericEntryBox.h"
 #include "SEditableComboBox.h"
diff --git a/Engine/Source/Runtime/Slate/Public/Widgets/Navigation/SWizard.h b/Engine/Source/Runtime/Slate/Public/Widgets/Navigation/SWizard.h
deleted file mode 100644
