Change 2122846 on 2014/07/01 00:51:24 by [REDACTED_PERSON][REDACTED_HANDLE]

	#BUN UnrealMath updates
	* added RotationAboutPointMatrix
	* Added static Make() methods for the various derived FMatrix types. Allows single expressions that return an FMatrix so we don't have to rely on implicit conversion (which will break some generic programming).
	* Simplified the implementation of FQuat::MakeFromEuler.

Change 2122848 on 2014/07/01 01:04:31 by [REDACTED_PERSON][REDACTED_HANDLE]

	Added SlateLayoutTransform to store a 2D translation and uniform scale (for FGeometry).
	#BUN Added SlateTransformCalculus for handling arbitrary transform manipulation. See documentation in header.

Change 2123889 on 2014/07/01 17:12:40 by [REDACTED_PERSON][REDACTED_HANDLE]

	Remove unncessary use of FGeometry.

Change 2221407 on 2014/07/16 17:58:20 by [REDACTED_PERSON][REDACTED_HANDLE]

	TransformCalculus changes
	* Move 2D versions of functions into SlateLayoutTransform header to separate it from the base 3D stuff.
	* FSlateLayoutTransform is scale then translate now instead of translate then scale.
	* TransformPosition -> TransformPoint
	* Added a version of Concatenate taking 4 transforms.
	* Update docs.

Change 2221408 on 2014/07/16 17:58:39 by [REDACTED_PERSON][REDACTED_HANDLE]

	Adjustments to Slate headers and PCH

Change 2221409 on 2014/07/16 17:59:07 by [REDACTED_PERSON][REDACTED_HANDLE]

	Transform calculus support for SlateRects

Change 2221459 on 2014/07/16 18:41:13 by [REDACTED_PERSON][REDACTED_HANDLE]

	#BUN SWindow changes
	* Added GetLocalToScreenTransform. (Screen space == Desktop space)
	** Basically contains the ApplicationScale and Window offset from the desktop.
	* Added GetLocalToWindowTransform (Window space == Desktop space without the window offset)
	* Added GetClientRectInScreen
	** Allows code to know exactly WHERE the client rect is in the window.
	* Refactored GetWindowGeometryInScreen and GetWindowGeometryInWindow to use the new methods above.
	* SPopupLayer::OnArrangChildren refactored to clarify transformational spaces and use transformation calculus to do the work.
	** This fixes PopupLayer to correctly account for the scale of the widget it is presenting a popup for.
	** MenuAnchors that that create new SWindows still do no correctly account for scale.

Change 2221464 on 2014/07/16 18:45:25 by [REDACTED_PERSON][REDACTED_HANDLE]

	Render Transforms initial checkin. Lots of cleanup to do, but this is functionally working for all basic widgets and draw elements in slate. Main feature todos:
	* restore pixel snapping.
	* handle clipping in the presence of rotation.

Change 2226298 on 2014/07/21 21:22:46 by [REDACTED_PERSON][REDACTED_HANDLE]

	Use a better method to get the window's rectangle.

Change 2227809 on 2014/07/22 20:41:57 by [REDACTED_PERSON][REDACTED_HANDLE]

	Add 2D version of TransformVector to SlateTransform.

Change 2232085 on 2014/07/25 16:32:15 by [REDACTED_PERSON][REDACTED_HANDLE]

	Implement CPU pixel snapping even under render transform. This is pretty CPU intensive, but a reasonable POC.

Change 2232090 on 2014/07/25 16:33:10 by [REDACTED_PERSON][REDACTED_HANDLE]

	Font cache fixes for padding.
	* Add 1 pixel of padding around all fonts.
	* Fix font cache to properly handle padding
	* Fix font cache to handle texture elements of zero size with non-zero padding.

Change 2232350 on 2014/07/25 19:50:11 by [REDACTED_PERSON][REDACTED_HANDLE]

	#BUN Updated FSlateTextureAtlas padding logic.
	* Added ESlateTextureAtlasPaddingStyle which lets you choose from three atlasing padding styles:
	  * NoPadding - don't apply any padding (used to be Padding == 0).
	  * DilateBorder - copy border edges to apply 1-pixel padding (used to be Padding == 1).
	  * PadWithZero - expand 1-pixel border and fill it with zeros (new style).
	* removed notion of Padding > 1 pixel from atlasing. It wasn't useful as we don't support mips or aniso filtering.
	  * Also the existing code didn't actually handle Padding > 1.
	* FontCache now uses PadWithZero padding so upon zooming text, the edges get smoother.

Change 2237962 on 2014/07/31 02:17:28 by [REDACTED_PERSON][REDACTED_HANDLE]

	TransformCalculus tweak.
	* Core Concatentate(T,T) no longer assumes the return type is T. Uses decltype(T.Concat(T)) instead.

Change 2237964 on 2014/07/31 02:21:31 by [REDACTED_PERSON][REDACTED_HANDLE]

	#BUN Added a full suite of transform types for 2D transformations that are transform calculus aware:
	* FSlateScale2D
	* FSlateShear2D
	* FSlateQuat2D
	* FSlateMatrix2x2

	Added appropriate overloads and ConcatenateRules for all combinations, making them compatible with FSlateTransform and FSlateLayoutTransform.

	FSlateTransform is now implemented in terms of a decomposed Matrix3x2.

Change 2237972 on 2014/07/31 02:30:45 by [REDACTED_PERSON][REDACTED_HANDLE]

	Remove 4x4 matrix from FSlateTransform

Change 2238004 on 2014/07/31 03:11:42 by [REDACTED_PERSON][REDACTED_HANDLE]

	FSlateTransform -> FSlateRenderTransform.

Change 2238714 on 2014/07/31 15:16:18 by [REDACTED_PERSON][REDACTED_HANDLE]

	#BUN Moved TransformCalculus into Core/Math.
	* Split into TransformCalculus, 2D, and 3D support.
	* FSlateRenderTransform is now just a typedef.
	* Renamed 2D transform types to remove "Slate" from their name.

Change 2240437 on 2014/08/01 19:08:59 by [REDACTED_PERSON][REDACTED_HANDLE]

	Restore funky logic in deprecated FGeometry::MakeChild implementations.
	* legacy FGeometry::MakeChild functions took OffsetFromParent in the local space of the widget. This therefore CANNOT be the layout transform offset, because that needs to be in parent space for the transform of P * S + T to work correctly.
	* Therefore, the legacy functions now expect the ChildOffset to be in child space and re-apply the child scale to put them back into parent space.
	* This fixes the NodePanels to properly handle zooming around the mouse cursor point.

Change 2242769 on 2014/08/04 13:44:33 by [REDACTED_PERSON][REDACTED_HANDLE]

	remove local layout transform storage from FGeometry.

Change 2242771 on 2014/08/04 13:45:59 by [REDACTED_PERSON][REDACTED_HANDLE]

	Update Declarative syntax support to call it RenderTransform instead of Transform.

Change 2245872 on 2014/08/06 14:45:44 by [REDACTED_PERSON][REDACTED_HANDLE]

	#BUN Add GetBottomRight function to FSlateRect.

Change 2246129 on 2014/08/06 17:32:04 by [REDACTED_PERSON][REDACTED_HANDLE]

	SThrobber now supports render transforms.
	Key is to not use deprecated members of FPaintGeometry, which should be treated as immutable (at least public members shouldn't be modified withotu modifying the new members as well).

Change 2247782 on 2014/08/07 19:12:57 by [REDACTED_PERSON][REDACTED_HANDLE]

	Clipping widgets now works under render transforms, FPaintGeometry legacy fixes.
	* Expanded FSlateVertex by 2 bytes to pass in local position.
	* Change clipcoords to be passed in local space.
	* Remove WorldPosition interpolator. Replaced by LocalPosition interpolator.

	FPaintGeometry fixes
	* FPaintGeometry now takes Size in LocalSpace unless using a legacy ctor. Avoid transform-and-back that was occurring when creating from an FGeometry, which already had it in local space, and rendering, which now wants it in local space.
	* FPaintGeometry again supports mutability when using legacy members (like DrawPosition). This allows legacy code to mutate an FPaintGeometry, which is committed to a render transform right before rendering.Hacky, but allows legacy code that treated FPaintGeometry as a plain-old-struct to work.
	* Removed a legacy ctor of FPaintGeoemtry that was not being used.

	Misc:
	* Turned off pixel snapping code in ElementBatcher because the FSlateVertex is doing it during construction.
	  * Still texture filtering artifacts to figure out, so the last word has not been spoken on this yet!
	* Tweaked FGeometry::MakeRoot to use the non-render transform ctor since it doesn't need to use it.
	* Removed /*WindowSpace*/ Size member from SlateDrawElement.
	* SlateDrawElement::GetRotationPoint vastly simplified by using Local space.
	* Moved slate cycle counters back to the top of the function in ElementBatcher.
	* Uninlined FSlateVertex ctors so I can tweak their rounding rules without a full recompile.

Change 2247790 on 2014/08/07 19:20:01 by [REDACTED_PERSON][REDACTED_HANDLE]

	SSlider no longer uses MakeRotatedBox, correctly handles render transforms.
	* Now uses a render transform to draw a vertical slider handle instead of MakeRotatedBox.
	* Since we don't support layout transforms, we also have to adjust the clip rect to account for the render-only transform, and make it act like a layout transform.

Change 2249106 on 2014/08/08 19:13:21 by [REDACTED_PERSON][REDACTED_HANDLE]

	Slate clipping now done in window space again.
	* This addresses subpixel jittering that occured on otherwise stable clip rects because local space was moving underneath, and transforming the clip rect to local space created occasional rounding errors.
	* Implemented a vectorized point-in-parallelogram check in the pixel shaders to do this. Now pass clip rect in slightly differently, and moved to a float16 to allow better clip precision under rotation. But the vertex data remains the same size.
	* Moved some more common transform work into FSlateVertex constructor to save duplication.
	* Removed construction of local clip rect as it's no longer needed except for text rendering.
	* Hoisted construction of local clip rect out of the text drawing character loop (whoops).
	* Precompute InverseLayoutTransform in draw code as we use that a lot now.
	* Fixed OGL renderer to work.

Change 2249123 on 2014/08/08 19:45:19 by [REDACTED_PERSON][REDACTED_HANDLE]

	SProgressBar no longer uses a temporary PaintGeometry to create the inner clipping rect for drawing, it directly inflates the clipping rect now. Also tried to adjust the coords to reduce jitter in the one-pixel padding.

Change 2251914 on 2014/08/11 16:26:12 by [REDACTED_PERSON][REDACTED_HANDLE]

	Fixing HitTest 2.0 to work with latest code.
	* Added AppendTransform to FGeometry to preserve immutability of the struct.

Change 2253023 on 2014/08/12 10:30:23 by [REDACTED_PERSON][REDACTED_HANDLE]

	Fixing Slate RHI Renderer to use new screenspace clipping with render transforms.

Change 2253090 on 2014/08/12 11:36:16 by [REDACTED_PERSON][REDACTED_HANDLE]

	Fix SClippingHorizontalBox to use immutable FGeometry position. Also fix to properly support scale in when it needs to clip children.

Change 2253701 on 2014/08/12 17:48:06 by [REDACTED_PERSON][REDACTED_HANDLE]

	#BUN Round viewport rect instead of trunc to match how slate does it now (fixes 1:1 pixel error on SViewports). Need to abstract slate rounding internals so clients don't have to guess.

Change 2255403 on 2014/08/13 19:08:25 by [REDACTED_PERSON][REDACTED_HANDLE]

	#BUN Added LayoutGeometry class.

Change 2255404 on 2014/08/13 19:10:40 by [REDACTED_PERSON][REDACTED_HANDLE]

	#BUN Made FGeometry public members const to catch all mutating uses of them. Added some hacks in the implementation to make it backward compatible.

	Removed the AccumulatedLayoutTransform, since it was just a copy of AbsolutePosition and Scale. Now all code creates the layout transform on the fly from those members. Still functionally the same.

	Added support for MakeChild using FLayoutGeometry.

Change 2255405 on 2014/08/13 19:12:39 by [REDACTED_PERSON][REDACTED_HANDLE]

	Changed SSplitter and Splitter2x2 to use FLayoutTransforms where possible. Now only calls ArrangeChildren() when rendering, avoiding accumulating the layout transforms and storing the render transforms altogether.

Change 2265162 on 2014/08/20 18:24:04 by [REDACTED_PERSON][REDACTED_HANDLE]

	Implement MakeRotatedBox in terms of a render transform. Remove all shader code associated with rotating.

Change 2268034 on 2014/08/22 17:40:05 by [REDACTED_PERSON][REDACTED_HANDLE]

	More 2D transform cleanup.
	* Removed a bunch of todo comments from the code, and converted them to more of a final form for main.
	* #BUN Removed FGeometry::CenteredPaintGeometryBelow as it was misleading, and users of the function were forced to mutate the resulting FPaintGeometry.
	* #BUN Removed FGeometry::GetRect() because it returned a rect in a weird space if the local scale != 1. Only one place in the code actually used it.

Change 2268096 on 2014/08/22 18:12:41 by [REDACTED_PERSON][REDACTED_HANDLE]

	Clean up more stuff in FPaintGeometry.
	* #BUN Removed ToSlateRect and Identity().
	* Removed storage of accumulated layout transform, since DrawPosition and DrawScale implicitly stored it.
	* Removed a bunch of todo comments and converted them to a more final form.

Change 2270989 on 2014/08/25 15:48:46 by [REDACTED_PERSON][REDACTED_HANDLE]

	Add support for Concatenate with 5 transforms.

Change 2271394 on 2014/08/25 19:12:53 by [REDACTED_PERSON][REDACTED_HANDLE]

	#BUN Added an OffsetBy function to FSlateRect

Change 2271395 on 2014/08/25 19:13:56 by [REDACTED_PERSON][REDACTED_HANDLE]

	Added a version of FSlateRotatedRect that uses regular floats, also added IsUnderLocation ToBoundingRect function for overlap testing.

Change 2271396 on 2014/08/25 19:14:38 by [REDACTED_PERSON][REDACTED_HANDLE]

	HitTest 2.0 now works with render transforms.

diff --git a/Engine/Plugins/Experimental/Paper2D/Source/Paper2DEditor/Private/Atlasing/PaperAtlasGenerator.cpp b/Engine/Plugins/Experimental/Paper2D/Source/Paper2DEditor/Private/Atlasing/PaperAtlasGenerator.cpp
--- a/Engine/Plugins/Experimental/Paper2D/Source/Paper2DEditor/Private/Atlasing/PaperAtlasGenerator.cpp
+++ b/Engine/Plugins/Experimental/Paper2D/Source/Paper2DEditor/Private/Atlasing/PaperAtlasGenerator.cpp
@@ -27,7 +27,7 @@ struct FSingleTexturePaperAtlas : public FSlateTextureAtlas
 {
 public:
 	FSingleTexturePaperAtlas(uint32 InWidth, uint32 InHeight)
-	: FSlateTextureAtlas(InWidth, InHeight, sizeof(FColor), 0)
+	: FSlateTextureAtlas(InWidth, InHeight, sizeof(FColor), ESlateTextureAtlasPaddingStyle::NoPadding)
 	{
 		//@TODO: check(StrideBytes >= 4*Width);
 	}
diff --git a/Engine/Shaders/SlateElementPixelShader.usf b/Engine/Shaders/SlateElementPixelShader.usf
--- a/Engine/Shaders/SlateElementPixelShader.usf
+++ b/Engine/Shaders/SlateElementPixelShader.usf
@@ -40,11 +40,11 @@ SamplerState ElementTextureSampler;
 
 struct VertexOut
 {
-	float4 Position				: SV_POSITION;
-	float4 Color				: COLOR0;
-	float4 TextureCoordinates	: TEXCOORD0;
-	float4 ClipCoords			: TEXCOORD1;
-	float2 WorldPosition		: TEXCOORD2;
+	float4 Position : SV_POSITION;
+	float4 Color : COLOR0;
+	float4 TextureCoordinates : TEXCOORD0;
+	float4 ClipOriginAndPos : TEXCOORD1;
+	float4 ClipExtents : TEXCOORD2;
 };
 	
 #if USE_MATERIALS
@@ -255,11 +255,49 @@ float4 ApplyColorVisionDeficiencyFilter(float4 Color)
 	return OutColor;
 }
 
+float cross(float2 a, float2 b)
+{
+	return a.x*b.y - a.y*b.x;
+}
+
+/**
+ * Given a point p and a parallelogram defined by point a and vectors b and c, determines in p is inside the parallelogram. 
+ * returns a 4-vector that can be used with the clip instruction.
+ */
+float4 PointInParallelogram(float2 p, float2 a, float4 bc)
+{
+	// unoptomized form:
+	//float2 o = p - a;
+	//float2 b = bc.xy;
+	//float2 c = bc.zw;
+	//float d = cross(b, c);
+	//float s = -cross(o, b) / d;
+	//float t = cross(o, c) / d;
+	// test for s and t between 0 and 1
+	//return float4(s, 1 - s, t, 1 - t);
+
+	float2 o = p - a;
+	// precompute 1/d
+	float invD = 1/cross(bc.xy, bc.zw);
+	// Compute an optimized o x b and o x c, leveraging that b and c are in the same vector register already (and free swizzles):
+	//   (o.x * b .y  - o.y * b .x, o.x *  c.y - o.y *  c.x) ==
+	//   (o.x * bc.y  - o.y * bc.x, o.x * bc.w - o.y * bc.z) ==
+	//    o.x * bc.yw - o.y * bc.xz
+	float2 st = (o.x * bc.yw - o.y * bc.xz) * float2(-invD, invD);
+	// test for s and t between 0 and 1
+	return float4(st, float2(1,1) - st);
+}
+
 float4 Main( VertexOut VIn ) : SV_Target0
 {
+	// Clip pixels which are outside of the clipping rect
+	float2 ClipOrigin = VIn.ClipOriginAndPos.xy;
+	float2 WindowPos = VIn.ClipOriginAndPos.zw;
+	float4 ClipTest = PointInParallelogram(WindowPos, ClipOrigin, VIn.ClipExtents);
+	
 #if !DEBUG_CLIPPING
 	// Clip pixels which are outside of the clipping rect
-	clip( float4( VIn.WorldPosition.x - VIn.ClipCoords.x, VIn.ClipCoords.z - VIn.WorldPosition.x, VIn.WorldPosition.y - VIn.ClipCoords.y, VIn.ClipCoords.w - VIn.WorldPosition.y ) );
+	clip(ClipTest);
 #endif
 
 	float4 OutColor;
@@ -295,8 +333,7 @@ float4 Main( VertexOut VIn ) : SV_Target0
 
 
 #if DEBUG_CLIPPING
-	float4 Diff = float4( VIn.WorldPosition.x - VIn.ClipCoords.x, VIn.ClipCoords.z - VIn.WorldPosition.x, VIn.WorldPosition.y - VIn.ClipCoords.y, VIn.ClipCoords.w - VIn.WorldPosition.y );
-	if( Diff.x < 0 || Diff.y < 0 || Diff.z < 0 || Diff.w < 0 )
+	if (any(ClipTest < 0))
 	{
 		// Shows red where any pixels discarded because part of the element was clipped
 		OutColor *= float4(1,0,0,1) ;
@@ -308,9 +345,13 @@ float4 Main( VertexOut VIn ) : SV_Target0
 
 float4 DebugOverdrawMain( VertexOut VIn ) : SV_Target0
 {
+
 #if !DEBUG_CLIPPING
 	// Clip pixels which are outside of the clipping rect
-	clip( float4( VIn.WorldPosition.x - VIn.ClipCoords.x, VIn.ClipCoords.z - VIn.WorldPosition.x, VIn.WorldPosition.y - VIn.ClipCoords.y, VIn.ClipCoords.w - VIn.WorldPosition.y ) );
+	float2 ClipOrigin = VIn.ClipOriginAndPos.xy;
+	float2 WindowPos = VIn.ClipOriginAndPos.zw;
+	float4 ClipTest = PointInParallelogram(WindowPos, ClipOrigin, VIn.ClipExtents);
+	clip(ClipTest);
 #endif
 
 	return .1;
diff --git a/Engine/Shaders/SlateVertexShader.usf b/Engine/Shaders/SlateVertexShader.usf
--- a/Engine/Shaders/SlateVertexShader.usf
+++ b/Engine/Shaders/SlateVertexShader.usf
@@ -4,43 +4,14 @@
 #include "Common.usf"
 
 float4x4 ViewProjection;
-float4 VertexShaderParams;
-
-/**
- * Rotates a point around another point 
- *
- * @param InPoint	The point to rotate
- * @param AboutPoint	The point to rotate about
- * @param Radians	The angle of rotation in radians
- * @return The rotated point
- */
-float2 RotatePoint( float2 InPoint, float2 AboutPoint, float Radians )
-{
-	if( Radians != 0.0f )
-	{
-		float CosAngle = cos(Radians);
-		float SinAngle = sin(Radians);
-
-		InPoint.x -= AboutPoint.x;
-		InPoint.y -= AboutPoint.y;
-
-		float X = (InPoint.x * CosAngle) - (InPoint.y * SinAngle);
-		float Y = (InPoint.x * SinAngle) + (InPoint.y * CosAngle);
-
-		return float2( X + AboutPoint.x, Y + AboutPoint.y );
-	}
-
-	return InPoint;
-}
-
 
 struct VertexOut
 {
 	float4 Position : SV_POSITION;
 	float4 Color : COLOR0;
 	float4 TextureCoordinates : TEXCOORD0;
-	float4 ClipCoords : TEXCOORD1;
-	float2 WorldPosition : TEXCOORD2;
+	float4 ClipOriginAndPos : TEXCOORD1;
+	float4 ClipExtents : TEXCOORD2;
 };
 
 VertexOut Main(
@@ -49,32 +20,30 @@ VertexOut Main(
 		D3D11: WARNING: ID3D11Device::CreateInputLayout: The provided input signature expects to read an element with SemanticName/Index: 'ATTRIBUTE'/1 and component(s) of the type 'float32'.  However, the matching entry in the Input Layout declaration, element[1], specifies mismatched format: 'R16G16B16A16_SINT'.  This is not an error, since behavior is well defined: The element format determines what data conversion algorithm gets applied before it shows up in a shader register. Independently, the shader input signature defines how the shader will interpret the data that has been placed in its input registers, with no change in the bits stored.  It is valid for the application to reinterpret data as a different type once it is in the vertex shader, so this warning is issued just in case reinterpretation was not intended by the author. [ STATE_CREATION WARNING #391: CREATEINPUTLAYOUT_TYPE_MISMATCH ]
 */
 #if ES2_PROFILE && COMPILER_GLSL_ES2
-	in float4 InClipCoords : ATTRIBUTE1,
-	in float2 InPosition : ATTRIBUTE2,
+	in float2 InPosition : ATTRIBUTE1,
+	in half2 InClipOrigin : ATTRIBUTE2,
+	in half4 InClipExtents : ATTRIBUTE3,
 #else
-	in int4 InClipCoords : ATTRIBUTE1,
-	in int2 InPosition : ATTRIBUTE2,
+	in int2 InPosition : ATTRIBUTE1,
+	in half2 InClipOrigin : ATTRIBUTE2,
+	in half4 InClipExtents : ATTRIBUTE3,
 #endif
-	in float4 InColor : ATTRIBUTE3
+	in float4 InColor : ATTRIBUTE4
 	)
 {
 	VertexOut VOut;
-	VOut.ClipCoords = InClipCoords;
-
-	float2 RotatedPoint = RotatePoint( InPosition.xy, VertexShaderParams.yz, VertexShaderParams.x );
-
-	// Position used for clipping, before transformation
-	VOut.WorldPosition = RotatedPoint.xy;
-
-	VOut.Position = mul(float4(RotatedPoint.xy,0,1), ViewProjection);
+	VOut.Position = mul(float4(InPosition.xy,0,1), ViewProjection);
 
 #if ES2_PROFILE && COMPILER_GLSL_ES2
 	// @todo-mobile: Fix this in the projection matrix
 	VOut.Position.y =  -VOut.Position.y;
 #endif
 
+
 	// TextureCoordinates contains both the first and second texture coordinates in xy and zw respectively.
 	VOut.TextureCoordinates = InTextureCoordinates;
+	VOut.ClipOriginAndPos = float4(InClipOrigin, InPosition.xy);
+	VOut.ClipExtents = InClipExtents;
 
 	VOut.Color = InColor FCOLOR_COMPONENT_SWIZZLE;
 
diff --git a/Engine/Shaders/StandaloneRenderer/D3D/SlateDefaultVertexShader.hlsl b/Engine/Shaders/StandaloneRenderer/D3D/SlateDefaultVertexShader.hlsl
--- a/Engine/Shaders/StandaloneRenderer/D3D/SlateDefaultVertexShader.hlsl
+++ b/Engine/Shaders/StandaloneRenderer/D3D/SlateDefaultVertexShader.hlsl
@@ -5,7 +5,6 @@
 cbuffer PerElementVSConstants
 {
 	matrix WorldViewProjection;
-	float4 VertexShaderParams;
 }
 
 struct VertexOut
@@ -13,57 +12,26 @@ struct VertexOut
 	float4 Position : SV_POSITION;
 	float4 Color : COLOR0;
 	float4 TextureCoordinates : TEXCOORD0;
-	float4 ClipCoords : TEXCOORD1;
-	float2 WorldPosition : TEXCOORD2;
+	float4 ClipOriginAndPos : TEXCOORD1;
+	float4 ClipExtents : TEXCOORD2;
 };
 
-
-
-/**
- * Rotates a point around another point 
- *
- * @param InPoint	The point to rotate
- * @param AboutPoint	The point to rotate about
- * @param Radians	The angle of rotation in radians
- * @return The rotated point
- */
-float2 RotatePoint( float2 InPoint, float2 AboutPoint, float Radians )
-{
-	if( Radians != 0.0f )
-	{
-		float CosAngle = cos(Radians);
-		float SinAngle = sin(Radians);
-
-		InPoint.x -= AboutPoint.x;
-		InPoint.y -= AboutPoint.y;
-
-		float X = (InPoint.x * CosAngle) - (InPoint.y * SinAngle);
-		float Y = (InPoint.x * SinAngle) + (InPoint.y * CosAngle);
-
-		return float2( X + AboutPoint.x, Y + AboutPoint.y );
-	}
-
-	return InPoint;
-}
-
-
 VertexOut Main(
 	in int2 InPosition : POSITION,
 	in float4 InTextureCoordinates : TEXCOORD0,
-	in uint4 InClipCoords : TEXCOORD1,
+	in half2 InClipOrigin : TEXCOORD1,
+	in half4 InClipExtents : TEXCOORD2,
 	in float4 InColor : COLOR0
 	) 
 {
 	VertexOut Out;
-	float2 RotatedPoint = RotatePoint( InPosition.xy, VertexShaderParams.yz, VertexShaderParams.x );
-	Out.WorldPosition = RotatedPoint;
 
-	Out.Position = mul( WorldViewProjection,float4( RotatedPoint, 0, 1 ) );
+	Out.Position = mul( WorldViewProjection,float4( InPosition.xy, 0, 1 ) );
 
-	Out.ClipCoords = InClipCoords;
 	Out.TextureCoordinates = InTextureCoordinates;
+	Out.ClipOriginAndPos = float4(InClipOrigin, InPosition.xy);
+	Out.ClipExtents = InClipExtents;
 	
-	// Swap r and b for d3d 11
 	Out.Color = InColor;
 
 	return Out;
diff --git a/Engine/Shaders/StandaloneRenderer/D3D/SlateElementPixelShader.hlsl b/Engine/Shaders/StandaloneRenderer/D3D/SlateElementPixelShader.hlsl
--- a/Engine/Shaders/StandaloneRenderer/D3D/SlateElementPixelShader.hlsl
+++ b/Engine/Shaders/StandaloneRenderer/D3D/SlateElementPixelShader.hlsl
@@ -28,10 +28,10 @@ struct VertexOut
 	float4 Position : SV_POSITION;
 	float4 Color : COLOR0;
 	float4 TextureCoordinates : TEXCOORD0;
-	float4 ClipCoords : TEXCOORD1;
-	float2 WorldPosition : TEXCOORD2;
+	float4 ClipOriginAndPos : TEXCOORD1;
+	float4 ClipExtents : TEXCOORD2;
 };
-	
+
 float3 Hue( float H )
 {
     float R = abs(H * 6 - 3) - 1;
@@ -150,13 +150,48 @@ float4 GetSplineElementColor( VertexOut InVertex )
 	return Color;
 }
 
+float cross(float2 a, float2 b)
+{
+	return a.x*b.y - a.y*b.x;
+}
+
+/**
+ * Given a point p and a parallelogram defined by point a and vectors b and c, determines in p is inside the parallelogram. 
+ * returns a 4-vector that can be used with the clip instruction.
+ */
+float4 PointInParallelogram(float2 p, float2 a, float4 bc)
+{
+	// unoptomized form:
+	//float2 o = p - a;
+	//float2 b = bc.xy;
+	//float2 c = bc.zw;
+	//float d = cross(b, c);
+	//float s = -cross(o, b) / d;
+	//float t = cross(o, c) / d;
+	// test for s and t between 0 and 1
+	//return float4(s, 1 - s, t, 1 - t);
+
+	float2 o = p - a;
+	// precompute 1/d
+	float invD = 1/cross(bc.xy, bc.zw);
+	// Compute an optimized o x b and o x c, leveraging that b and c are in the same vector register already (and free swizzles):
+	//   (o.x * b .y  - o.y * b .x, o.x *  c.y - o.y *  c.x) ==
+	//   (o.x * bc.y  - o.y * bc.x, o.x * bc.w - o.y * bc.z) ==
+	//    o.x * bc.yw - o.y * bc.xz
+	float2 st = (o.x * bc.yw - o.y * bc.xz) * float2(-invD, invD);
+	// test for s and t between 0 and 1
+	return float4(st, float2(1,1) - st);
+}
 
 float4 Main( VertexOut InVertex ) : SV_Target
 {
 	// Clip pixels which are outside of the clipping rect
-	float4 ClipCoords = InVertex.ClipCoords;
-	float2 SSPosition = InVertex.WorldPosition.xy;
-	clip( float4(SSPosition.x - ClipCoords.x, SSPosition.y - ClipCoords.y, ClipCoords.z - SSPosition.x, ClipCoords.w - SSPosition.y ) );
+	float2 ClipOrigin = InVertex.ClipOriginAndPos.xy;
+	float2 WindowPos = InVertex.ClipOriginAndPos.zw;
+	float4 ClipTest = PointInParallelogram(WindowPos, ClipOrigin, InVertex.ClipExtents);
+	
+	clip(ClipTest);
+	//float4 OutColorTint = any(ClipTest < 0) ? float4(1, 0.5, 0.5, 0.5) : float4(1, 1, 1, 1);
 
 	float4 OutColor;
 
@@ -193,6 +228,8 @@ float4 Main( VertexOut InVertex ) : SV_Target
 		OutColor.rgb = lerp( OutColor.rgb, Grayish, clamp( distance( OutColor.rgb, Grayish ), 0, .8)  );
 	}
 
-	return OutColor;
+	return OutColor
+		//* OutColorTint
+		;
 }
 
diff --git a/Engine/Shaders/StandaloneRenderer/OpenGL/SlateElementPixelShader.glsl b/Engine/Shaders/StandaloneRenderer/OpenGL/SlateElementPixelShader.glsl
--- a/Engine/Shaders/StandaloneRenderer/OpenGL/SlateElementPixelShader.glsl
+++ b/Engine/Shaders/StandaloneRenderer/OpenGL/SlateElementPixelShader.glsl
@@ -23,8 +23,8 @@ uniform int ShaderType;
 uniform sampler2D ElementTexture;
 
 varying vec4 TexCoords;
-varying vec2 SSPosition;
-varying vec4 ClipCoords;
+varying vec4 ClipOriginAndPos;
+varying vec4 ClipExtents;
 varying vec4 Color;
 
 vec4 GetFontElementColor()
@@ -99,6 +99,7 @@ vec4 GetBorderElementColor()
 
 vec4 GetSplineElementColor()
 {
+	vec2 SSPosition = ClipOriginAndPos.zw;
 	vec4 InTexCoords = TexCoords;
 	float Width = MarginUVs.x;
 	float Radius = MarginUVs.y;
@@ -146,16 +147,57 @@ vec4 GetSplineElementColor()
 	return InColor;
 }
 
+void clip(vec4 ClipTest)
+{
+	if( any( lessThan(ClipTest, vec4(0,0,0,0) ) ) ) discard;
+}
+
+float cross(vec2 a, vec2 b)
+{
+	return a.x*b.y - a.y*b.x;
+}
+
+/**
+ * Given a point p and a parallelogram defined by point a and vectors b and c, determines in p is inside the parallelogram. 
+ * returns a 4-vector that can be used with the clip instruction.
+ */
+vec4 PointInParallelogram(vec2 p, vec2 a, vec4 bc)
+{
+	// unoptomized form:
+	//vec2 o = p - a;
+	//vec2 b = bc.xy;
+	//vec2 c = bc.zw;
+	//float d = cross(b, c);
+	//float s = -cross(o, b) / d;
+	//float t = cross(o, c) / d;
+	// test for s and t between 0 and 1
+	//return vec4(s, 1 - s, t, 1 - t);
+
+	vec2 o = p - a;
+	// precompute 1/d
+	float invD = 1/cross(bc.xy, bc.zw);
+	// Compute an optimized o x b and o x c, leveraging that b and c are in the same vector register already (and free swizzles):
+	//   (o.x * b .y  - o.y * b .x, o.x *  c.y - o.y *  c.x) ==
+	//   (o.x * bc.y  - o.y * bc.x, o.x * bc.w - o.y * bc.z) ==
+	//    o.x * bc.yw - o.y * bc.xz
+	vec2 st = (o.x * bc.yw - o.y * bc.xz) * vec2(-invD, invD);
+	// test for s and t between 0 and 1
+	return vec4(st, vec2(1,1) - st);
+}
+
 void main()
 {
 	// Clip pixels which are outside of the clipping rect
-	vec4 Clip = vec4( SSPosition.x - ClipCoords.x, ClipCoords.z - SSPosition.x, SSPosition.y - ClipCoords.y, ClipCoords.w - SSPosition.y );
-	if( any( lessThan(Clip, vec4(0,0,0,0) ) ) )
-	{
-		discard;
-	}
+	vec2 ClipOrigin = ClipOriginAndPos.xy;
+	vec2 WindowPos = ClipOriginAndPos.zw;
+	vec4 ClipTest = PointInParallelogram(WindowPos, ClipOrigin, ClipExtents);
+	
+	clip(ClipTest);
+
+	//vec4 OutColorTint = any(lessThan(ClipTest, vec4(0,0,0,0))) ? vec4(1, 0.5, 0.5, 0.5) : vec4(1, 1, 1, 1);
 
 	vec4 OutColor;
+
 	if( ShaderType == ST_Default )
 	{
 		OutColor = GetDefaultElementColor();
@@ -192,5 +234,7 @@ void main()
 		OutColor.rgb = mix( OutColor.rgb, Grayish, clamp( distance( OutColor.rgb, Grayish ), 0.0, 0.8)  );
 	}
 
-	gl_FragColor = OutColor.bgra;
+	gl_FragColor = OutColor.bgra 
+		//* OutColorTint
+		;
 }
diff --git a/Engine/Shaders/StandaloneRenderer/OpenGL/SlateVertexShader.glsl b/Engine/Shaders/StandaloneRenderer/OpenGL/SlateVertexShader.glsl
--- a/Engine/Shaders/StandaloneRenderer/OpenGL/SlateVertexShader.glsl
+++ b/Engine/Shaders/StandaloneRenderer/OpenGL/SlateVertexShader.glsl
@@ -3,55 +3,26 @@
 // #version 120 at the beginning is added in FSlateOpenGLShader::CompileShader()
 
 uniform mat4 ViewProjectionMatrix;
-uniform vec4 VertexShaderParams;
 
 // Per vertex
 attribute vec4 InTexCoords;
-attribute vec4 InClipCoords;
 attribute vec2 InPosition;
+attribute vec2 InClipOrigin;
+attribute vec4 InClipExtents;
 attribute vec4 InColor;
 
 // Between vertex and pixel shader
 varying vec4 TexCoords;
-varying vec2 SSPosition;
-varying vec4 ClipCoords;
+varying vec4 ClipOriginAndPos;
+varying vec4 ClipExtents;
 varying vec4 Color;
 
-/**
- * Rotates a point around another point 
- *
- * @param InPoint	The point to rotate
- * @param AboutPoint	The point to rotate about
- * @param Radians	The angle of rotation in radians
- * @return The rotated point
- */
-vec2 RotatePoint( vec2 InPoint, vec2 AboutPoint, float Radians )
-{
-	if( Radians != 0.0 )
-	{
-		float CosAngle = cos(Radians);
-		float SinAngle = sin(Radians);
-
-		InPoint.x -= AboutPoint.x;
-		InPoint.y -= AboutPoint.y;
-
-		float X = (InPoint.x * CosAngle) - (InPoint.y * SinAngle);
-		float Y = (InPoint.x * SinAngle) + (InPoint.y * CosAngle);
-
-		return vec2( X + AboutPoint.x, Y + AboutPoint.y );
-	}
-
-	return InPoint;
-}
-
 void main()
 {
-	vec2 RotatedPosition = RotatePoint( InPosition, VertexShaderParams.yz, VertexShaderParams.x );
-
-	SSPosition = RotatedPosition;
-	ClipCoords = InClipCoords;
 	TexCoords = InTexCoords;
+	ClipOriginAndPos = vec4(InClipOrigin, InPosition);
+	ClipExtents = InClipExtents;
 	Color = InColor;
 
-	gl_Position = ViewProjectionMatrix * vec4( RotatedPosition.xy, 0, 1 );
+	gl_Position = ViewProjectionMatrix * vec4( InPosition, 0, 1 );
 }
diff --git a/Engine/Source/Developer/LogVisualizer/Private/SLogBar.cpp b/Engine/Source/Developer/LogVisualizer/Private/SLogBar.cpp
--- a/Engine/Source/Developer/LogVisualizer/Private/SLogBar.cpp
+++ b/Engine/Source/Developer/LogVisualizer/Private/SLogBar.cpp
@@ -46,8 +46,8 @@ int32 SLogBar::OnPaint( const FPaintArgs& Args, const FGeometry& AllottedGeometr
 
 	// Paint inside the border only. 
 	const FVector2D BorderPadding = FEditorStyle::GetVector("ProgressBar.BorderPadding");
-	FPaintGeometry ForegroundPaintGeometry = AllottedGeometry.ToInflatedPaintGeometry( -BorderPadding );
-	const FSlateRect ForegroundClippingRect = ForegroundPaintGeometry.ToSlateRect().IntersectionWith( MyClippingRect );
+	const FSlateRect ForegroundClippingRect = AllottedGeometry.GetClippingRect().InsetBy(FMargin(BorderPadding.X, BorderPadding.Y)).IntersectionWith(MyClippingRect);
+
 
 	const float LogBarWidth = AllottedGeometry.Size.X;
 
diff --git a/Engine/Source/Developer/MeshUtilities/Private/MeshUtilities.cpp b/Engine/Source/Developer/MeshUtilities/Private/MeshUtilities.cpp
--- a/Engine/Source/Developer/MeshUtilities/Private/MeshUtilities.cpp
+++ b/Engine/Source/Developer/MeshUtilities/Private/MeshUtilities.cpp
@@ -3373,7 +3373,7 @@ private:
 	struct FLightmapAtlas : public FSlateTextureAtlas
 	{
 		FLightmapAtlas(uint32 InWidth)
-			: FSlateTextureAtlas(InWidth, InWidth, 0, 0)
+			: FSlateTextureAtlas(InWidth, InWidth, 0, ESlateTextureAtlasPaddingStyle::NoPadding)
 		{}
 
 		const FAtlasedTextureSlot* AddLightmap(uint32 InWidth)
diff --git a/Engine/Source/Developer/MovieSceneTools/Private/DirectorTrackEditor.cpp b/Engine/Source/Developer/MovieSceneTools/Private/DirectorTrackEditor.cpp
--- a/Engine/Source/Developer/MovieSceneTools/Private/DirectorTrackEditor.cpp
+++ b/Engine/Source/Developer/MovieSceneTools/Private/DirectorTrackEditor.cpp
@@ -226,10 +226,10 @@ int32 FShotSection::OnPaintSection( const FGeometry& AllottedGeometry, const FSl
 	{
 		for (int32 i = 0; i < Thumbnails.Num(); ++i)
 		{
-			FGeometry TruncatedGeometry = AllottedGeometry;
-			TruncatedGeometry.Position.X += i * ThumbnailWidth;
-			TruncatedGeometry.AbsolutePosition.X += i * ThumbnailWidth;
-			TruncatedGeometry.Size.X = ThumbnailWidth;
+			FGeometry TruncatedGeometry = AllottedGeometry.MakeChild(
+				FVector2D(ThumbnailWidth, AllottedGeometry.Size.Y),
+				FSlateLayoutTransform(AllottedGeometry.Scale, FVector2D(i * ThumbnailWidth, 0.0f))
+				);
 			ensure(TruncatedGeometry.Size.Y == AnimatableShotToolConstants::DirectorTrackHeight);
 
 			FSlateDrawElement::MakeViewport(
diff --git a/Engine/Source/Developer/SlateReflector/Private/Widgets/SWidgetReflector.cpp b/Engine/Source/Developer/SlateReflector/Private/Widgets/SWidgetReflector.cpp
--- a/Engine/Source/Developer/SlateReflector/Private/Widgets/SWidgetReflector.cpp
+++ b/Engine/Source/Developer/SlateReflector/Private/Widgets/SWidgetReflector.cpp
@@ -232,18 +232,22 @@ int32 SWidgetReflector::VisualizePickAsRectangles( const FWidgetPath& InWidgetsT
 		const float ColorFactor = static_cast<float>(WidgetIndex)/InWidgetsToVisualize.Widgets.Num();
 		const FLinearColor Tint(1.0f - ColorFactor, ColorFactor, 0.0f, 1.0f);
 
+		// The FGeometry we get is from a WidgetPath, so it's rooted in desktop space.
+		// We need to APPEND a transform to the Geometry to essentially undo this root transform
+		// and get us back into Window Space.
+		// This is nonstandard so we have to go through some hoops and a specially exposed method 
+		// in FPaintGeometry to allow appending layout transforms.
+		FPaintGeometry WindowSpaceGeometry = WidgetGeometry.Geometry.ToPaintGeometry();
+		WindowSpaceGeometry.AppendTransform(TransformCast<FSlateLayoutTransform>(Inverse(InWidgetsToVisualize.TopLevelWindow->GetPositionInScreen())));
+
 		FSlateDrawElement::MakeBox(
 			OutDrawElements,
 			++LayerId,
-			FPaintGeometry(
-				WidgetGeometry.Geometry.AbsolutePosition - InWidgetsToVisualize.TopLevelWindow->GetPositionInScreen(),
-				WidgetGeometry.Geometry.Size * WidgetGeometry.Geometry.Scale,
-				WidgetGeometry.Geometry.Scale),
-				FCoreStyle::Get().GetBrush(TEXT("Debug.Border")),
-				InWidgetsToVisualize.TopLevelWindow->GetClippingRectangleInWindow(),
-				ESlateDrawEffect::None,
-				FMath::Lerp( TopmostWidgetColor, LeafmostWidgetColor, ColorFactor
-			)
+			WindowSpaceGeometry,
+			FCoreStyle::Get().GetBrush(TEXT("Debug.Border")),
+			InWidgetsToVisualize.TopLevelWindow->GetClippingRectangleInWindow(),
+			ESlateDrawEffect::None,
+			FMath::Lerp( TopmostWidgetColor, LeafmostWidgetColor, ColorFactor )
 		);
 	}
 
@@ -258,14 +262,18 @@ int32 SWidgetReflector::VisualizeSelectedNodesAsRectangles( const TArray<TShared
 		const TSharedPtr<FReflectorNode>& NodeToDraw = InNodesToDraw[NodeIndex];
 		const FLinearColor Tint(0.0f, 1.0f, 0.0f);
 
+		// The FGeometry we get is from a WidgetPath, so it's rooted in desktop space.
+		// We need to APPEND a transform to the Geometry to essentially undo this root transform
+		// and get us back into Window Space.
+		// This is nonstandard so we have to go through some hoops and a specially exposed method 
+		// in FPaintGeometry to allow appending layout transforms.
+		FPaintGeometry WindowSpaceGeometry = NodeToDraw->Geometry.ToPaintGeometry();
+		WindowSpaceGeometry.AppendTransform(TransformCast<FSlateLayoutTransform>(Inverse(VisualizeInWindow->GetPositionInScreen())));
+
 		FSlateDrawElement::MakeBox(
 			OutDrawElements,
 			++LayerId,
-			FPaintGeometry(
-				NodeToDraw->Geometry.AbsolutePosition - VisualizeInWindow->GetPositionInScreen(),
-				NodeToDraw->Geometry.Size * NodeToDraw->Geometry.Scale,
-				NodeToDraw->Geometry.Scale
-			),
+			WindowSpaceGeometry,
 			FCoreStyle::Get().GetBrush(TEXT("Debug.Border")),
 			VisualizeInWindow->GetClippingRectangleInWindow(),
 			ESlateDrawEffect::None,
diff --git a/Engine/Source/Developer/StandaloneRenderer/Private/OpenGL/SlateOpenGLRenderer.h b/Engine/Source/Developer/StandaloneRenderer/Private/OpenGL/SlateOpenGLRenderer.h
--- a/Engine/Source/Developer/StandaloneRenderer/Private/OpenGL/SlateOpenGLRenderer.h
+++ b/Engine/Source/Developer/StandaloneRenderer/Private/OpenGL/SlateOpenGLRenderer.h
@@ -5,6 +5,7 @@
 
 class FSlateOpenGLTextureManager;
 class FSlateOpenGLRenderingPolicy;
+class FSlateElementBatcher;
 
 #if !(UE_BUILD_SHIPPING || UE_BUILD_TEST)
 #define CHECK_GL_ERRORS \
diff --git a/Engine/Source/Developer/StandaloneRenderer/Private/OpenGL/SlateOpenGLRenderingPolicy.cpp b/Engine/Source/Developer/StandaloneRenderer/Private/OpenGL/SlateOpenGLRenderingPolicy.cpp
--- a/Engine/Source/Developer/StandaloneRenderer/Private/OpenGL/SlateOpenGLRenderingPolicy.cpp
+++ b/Engine/Source/Developer/StandaloneRenderer/Private/OpenGL/SlateOpenGLRenderingPolicy.cpp
@@ -188,8 +188,6 @@ void FSlateOpenGLRenderingPolicy::DrawElements( const FMatrix& ViewProjectionMat
 		}
 #endif
 
-		ElementProgram.SetVertexShaderParams( RenderBatch.ShaderParams.VertexParams );
-
 		ElementProgram.SetShaderType( RenderBatch.ShaderType );
 		ElementProgram.SetMarginUVs( RenderBatch.ShaderParams.PixelParams );
 		ElementProgram.SetDrawEffects( RenderBatch.DrawEffects );
@@ -244,16 +242,20 @@ void FSlateOpenGLRenderingPolicy::DrawElements( const FMatrix& ViewProjectionMat
 		glVertexAttribPointer( 0, 4, GL_FLOAT, GL_FALSE, Stride, BUFFER_OFFSET(Stride*BaseVertexIndex+Offset) );
 
 		glEnableVertexAttribArray(1);
-		Offset = STRUCT_OFFSET( FSlateVertex, ClipCoords );
-		glVertexAttribPointer( 1, 4, GL_UNSIGNED_SHORT, GL_FALSE, Stride, BUFFER_OFFSET(Stride*BaseVertexIndex+Offset) );
+		Offset = STRUCT_OFFSET( FSlateVertex, Position );
+		glVertexAttribPointer( 1, 2, GL_SHORT, GL_FALSE, Stride, BUFFER_OFFSET(Stride*BaseVertexIndex+Offset) );
 
 		glEnableVertexAttribArray(2);
-		Offset = STRUCT_OFFSET( FSlateVertex, Position );
-		glVertexAttribPointer( 2, 2, GL_SHORT, GL_FALSE, Stride, BUFFER_OFFSET(Stride*BaseVertexIndex+Offset) );
+		Offset = STRUCT_OFFSET( FSlateVertex, ClipRect );
+		glVertexAttribPointer( 2, 2, GL_HALF_FLOAT, GL_FALSE, Stride, BUFFER_OFFSET(Stride*BaseVertexIndex+Offset) );
 
 		glEnableVertexAttribArray(3);
+		Offset = STRUCT_OFFSET( FSlateVertex, ClipRect ) + STRUCT_OFFSET( FSlateRotatedRectFloat16, ExtentX );
+		glVertexAttribPointer( 3, 4, GL_HALF_FLOAT, GL_FALSE, Stride, BUFFER_OFFSET(Stride*BaseVertexIndex+Offset) );
+
+		glEnableVertexAttribArray(4);
 		Offset = STRUCT_OFFSET( FSlateVertex, Color );
-		glVertexAttribPointer( 3, 4, GL_UNSIGNED_BYTE, GL_TRUE, Stride, BUFFER_OFFSET(Stride*BaseVertexIndex+Offset) );
+		glVertexAttribPointer( 4, 4, GL_UNSIGNED_BYTE, GL_TRUE, Stride, BUFFER_OFFSET(Stride*BaseVertexIndex+Offset) );
 
 
 		
diff --git a/Engine/Source/Developer/StandaloneRenderer/Private/OpenGL/SlateOpenGLShaders.cpp b/Engine/Source/Developer/StandaloneRenderer/Private/OpenGL/SlateOpenGLShaders.cpp
--- a/Engine/Source/Developer/StandaloneRenderer/Private/OpenGL/SlateOpenGLShaders.cpp
+++ b/Engine/Source/Developer/StandaloneRenderer/Private/OpenGL/SlateOpenGLShaders.cpp
@@ -210,9 +210,10 @@ void FSlateOpenGLShaderProgram::LinkShaders( const FSlateOpenGLVS& VertexShader,
 
 	// Set up attribute locations for per vertex data
 	glBindAttribLocation(ProgramID, 0, "InTexCoords");
-	glBindAttribLocation(ProgramID, 1, "InClipCoords");
-	glBindAttribLocation(ProgramID, 2, "InPosition");
-	glBindAttribLocation(ProgramID, 3, "InColor");
+	glBindAttribLocation(ProgramID, 1, "InPosition");
+	glBindAttribLocation(ProgramID, 2, "InClipOrigin");
+	glBindAttribLocation(ProgramID, 3, "InClipExtents");
+	glBindAttribLocation(ProgramID, 4, "InColor");
 
 	// Link the shaders
 	glLinkProgram( ProgramID );
diff --git a/Engine/Source/Developer/StandaloneRenderer/Private/Windows/D3D/SlateD3DRenderingPolicy.cpp b/Engine/Source/Developer/StandaloneRenderer/Private/Windows/D3D/SlateD3DRenderingPolicy.cpp
--- a/Engine/Source/Developer/StandaloneRenderer/Private/Windows/D3D/SlateD3DRenderingPolicy.cpp
+++ b/Engine/Source/Developer/StandaloneRenderer/Private/Windows/D3D/SlateD3DRenderingPolicy.cpp
@@ -221,7 +221,6 @@ void FSlateD3D11RenderingPolicy::DrawElements( const FMatrix& ViewProjectionMatr
 
 		const FShaderParams& ShaderParams = RenderBatch.ShaderParams;
 
-		VertexShader->SetShaderParams( ShaderParams.VertexParams );
 		VertexShader->BindParameters();
 		
 		if( DrawFlags & ESlateBatchDrawFlag::NoBlending )
diff --git a/Engine/Source/Developer/StandaloneRenderer/Private/Windows/D3D/SlateD3DShaders.cpp b/Engine/Source/Developer/StandaloneRenderer/Private/Windows/D3D/SlateD3DShaders.cpp
--- a/Engine/Source/Developer/StandaloneRenderer/Private/Windows/D3D/SlateD3DShaders.cpp
+++ b/Engine/Source/Developer/StandaloneRenderer/Private/Windows/D3D/SlateD3DShaders.cpp
@@ -241,12 +241,13 @@ FSlateDefaultVS::FSlateDefaultVS()
 	D3D11_INPUT_ELEMENT_DESC Layout[] = 
 	{
 		{ "TEXCOORD",	0, DXGI_FORMAT_R32G32B32A32_FLOAT,	0, 0,								D3D11_INPUT_PER_VERTEX_DATA,	0 },
-		{ "TEXCOORD",	1, DXGI_FORMAT_R16G16B16A16_UINT,	0, D3D11_APPEND_ALIGNED_ELEMENT,	D3D11_INPUT_PER_VERTEX_DATA,	0 },
 		{ "POSITION",	0, DXGI_FORMAT_R16G16_SINT,			0, D3D11_APPEND_ALIGNED_ELEMENT,	D3D11_INPUT_PER_VERTEX_DATA,	0 },
+		{ "TEXCOORD",	1, DXGI_FORMAT_R16G16_FLOAT,		0, D3D11_APPEND_ALIGNED_ELEMENT,	D3D11_INPUT_PER_VERTEX_DATA,	0 },
+		{ "TEXCOORD",	2, DXGI_FORMAT_R16G16B16A16_FLOAT,	0, D3D11_APPEND_ALIGNED_ELEMENT,	D3D11_INPUT_PER_VERTEX_DATA,	0 },
 		{ "COLOR",		0, DXGI_FORMAT_B8G8R8A8_UNORM,		0, D3D11_APPEND_ALIGNED_ELEMENT,	D3D11_INPUT_PER_VERTEX_DATA,	0 },
 	};
 
-	Create( FString::Printf( TEXT("%sShaders/StandaloneRenderer/D3D/SlateDefaultVertexShader.hlsl"), *FPaths::EngineDir() ), TEXT("Main"), TEXT("vs_4_0"), Layout, 4 );
+	Create( FString::Printf( TEXT("%sShaders/StandaloneRenderer/D3D/SlateDefaultVertexShader.hlsl"), *FPaths::EngineDir() ), TEXT("Main"), TEXT("vs_4_0"), Layout, ARRAY_COUNT(Layout) );
 
 }
 
diff --git a/Engine/Source/Developer/StandaloneRenderer/Private/Windows/D3D/SlateD3DTextureManager.cpp b/Engine/Source/Developer/StandaloneRenderer/Private/Windows/D3D/SlateD3DTextureManager.cpp
--- a/Engine/Source/Developer/StandaloneRenderer/Private/Windows/D3D/SlateD3DTextureManager.cpp
+++ b/Engine/Source/Developer/StandaloneRenderer/Private/Windows/D3D/SlateD3DTextureManager.cpp
@@ -341,6 +341,7 @@ FSlateShaderResourceProxy* FSlateD3DTextureManager::GenerateTextureResource( con
 		const uint32 AtlasSize = 1024;
 		// 4 bytes per pixel
 		const uint32 AtlasStride = sizeof(uint8)*4; 
+		// always use one pixel padding.
 		const uint8 Padding = 1;
 		const FAtlasedTextureSlot* NewSlot = NULL;
 
@@ -361,7 +362,7 @@ FSlateShaderResourceProxy* FSlateD3DTextureManager::GenerateTextureResource( con
 		if( !NewSlot )
 		{
 			// A new slot in the atlas could not be found, make a new atlas and add the texture to it
-			Atlas = new FSlateTextureAtlasD3D( AtlasSize, AtlasSize, AtlasStride, Padding );
+			Atlas = new FSlateTextureAtlasD3D( AtlasSize, AtlasSize, AtlasStride, ESlateTextureAtlasPaddingStyle::DilateBorder );
 			TextureAtlases.Add( Atlas );
 			NewSlot = TextureAtlases.Last()->AddTexture( Width, Height, Info.TextureData->GetRawBytes() );
 		}
diff --git a/Engine/Source/Developer/StandaloneRenderer/Private/Windows/D3D/SlateD3DTextures.cpp b/Engine/Source/Developer/StandaloneRenderer/Private/Windows/D3D/SlateD3DTextures.cpp
--- a/Engine/Source/Developer/StandaloneRenderer/Private/Windows/D3D/SlateD3DTextures.cpp
+++ b/Engine/Source/Developer/StandaloneRenderer/Private/Windows/D3D/SlateD3DTextures.cpp
@@ -35,8 +35,8 @@ void FSlateD3DTexture::Init( DXGI_FORMAT InFormat, D3D11_SUBRESOURCE_DATA* Inita
 	check( SUCCEEDED(Hr) );
 }
 
-FSlateTextureAtlasD3D::FSlateTextureAtlasD3D( uint32 Width, uint32 Height, uint32 StrideBytes, uint8 Padding )
-	: FSlateTextureAtlas( Width, Height, StrideBytes, Padding )
+FSlateTextureAtlasD3D::FSlateTextureAtlasD3D( uint32 Width, uint32 Height, uint32 StrideBytes, ESlateTextureAtlasPaddingStyle PaddingStyle )
+	: FSlateTextureAtlas( Width, Height, StrideBytes, PaddingStyle )
 	, AtlasTexture( new FSlateD3DTexture( Width, Height ) )
 {
 }
diff --git a/Engine/Source/Developer/StandaloneRenderer/Private/Windows/D3D/SlateD3DTextures.h b/Engine/Source/Developer/StandaloneRenderer/Private/Windows/D3D/SlateD3DTextures.h
--- a/Engine/Source/Developer/StandaloneRenderer/Private/Windows/D3D/SlateD3DTextures.h
+++ b/Engine/Source/Developer/StandaloneRenderer/Private/Windows/D3D/SlateD3DTextures.h
@@ -38,7 +38,7 @@ private:
 class FSlateTextureAtlasD3D : public FSlateTextureAtlas
 {
 public:
-	FSlateTextureAtlasD3D( uint32 Width, uint32 Height, uint32 StrideBytes, uint8 Padding );
+	FSlateTextureAtlasD3D( uint32 Width, uint32 Height, uint32 StrideBytes, ESlateTextureAtlasPaddingStyle PaddingStyle );
 	~FSlateTextureAtlasD3D();
 
 	void InitAtlasTexture( int32 Index );
diff --git a/Engine/Source/Developer/TaskGraph/Private/SGraphBar.cpp b/Engine/Source/Developer/TaskGraph/Private/SGraphBar.cpp
--- a/Engine/Source/Developer/TaskGraph/Private/SGraphBar.cpp
+++ b/Engine/Source/Developer/TaskGraph/Private/SGraphBar.cpp
@@ -68,8 +68,7 @@ int32 SGraphBar::OnPaint( const FPaintArgs& Args, const FGeometry& AllottedGeome
 
 	// Paint inside the border only. 
 	const FVector2D BorderPadding = FEditorStyle::GetVector("ProgressBar.BorderPadding");
-	FPaintGeometry ForegroundPaintGeometry = AllottedGeometry.ToInflatedPaintGeometry( -BorderPadding );
-	const FSlateRect ForegroundClippingRect = ForegroundPaintGeometry.ToSlateRect().IntersectionWith( MyClippingRect );
+	const FSlateRect ForegroundClippingRect = AllottedGeometry.GetClippingRect().InsetBy(FMargin(BorderPadding.X, BorderPadding.Y)).IntersectionWith(MyClippingRect);
 
 	FSlateDrawElement::MakeBox(
 		OutDrawElements,
diff --git a/Engine/Source/Developer/TaskGraph/Private/STimeline.cpp b/Engine/Source/Developer/TaskGraph/Private/STimeline.cpp
--- a/Engine/Source/Developer/TaskGraph/Private/STimeline.cpp
+++ b/Engine/Source/Developer/TaskGraph/Private/STimeline.cpp
@@ -32,8 +32,7 @@ int32 STimeline::OnPaint( const FPaintArgs& Args, const FGeometry& AllottedGeome
 
 	// Paint inside the border only. 
 	const FVector2D BorderPadding = FEditorStyle::GetVector("ProgressBar.BorderPadding");
-	FPaintGeometry ForegroundPaintGeometry = AllottedGeometry.ToInflatedPaintGeometry( -BorderPadding );
-	const FSlateRect ForegroundClippingRect = ForegroundPaintGeometry.ToSlateRect().IntersectionWith( MyClippingRect );
+	const FSlateRect ForegroundClippingRect = AllottedGeometry.GetClippingRect().InsetBy(FMargin(BorderPadding.X, BorderPadding.Y)).IntersectionWith(MyClippingRect);
 
 	const float OffsetX = DrawingOffsetX; // BorderPadding.X
 	const float Width = DrawingGeometry.Size.X; // AllottedGeometry.Size.X - - 2.0f * BorderPadding.X
diff --git a/Engine/Source/Editor/GraphEditor/Private/SGraphPanel.cpp b/Engine/Source/Editor/GraphEditor/Private/SGraphPanel.cpp
--- a/Engine/Source/Editor/GraphEditor/Private/SGraphPanel.cpp
+++ b/Engine/Source/Editor/GraphEditor/Private/SGraphPanel.cpp
@@ -219,11 +219,7 @@ int32 SGraphPanel::OnPaint( const FPaintArgs& Args, const FGeometry& AllottedGeo
 
 						if (bShowCommentBubble)
 						{
-							FGeometry CommentGeometry = CurWidget.Geometry;
-							if (!bScaleComments)
-							{
-								CommentGeometry.Scale = 1.0f;
-							}
+							FGeometry CommentGeometry = CurWidget.Geometry.MakeChild(CurWidget.Geometry.Size, FSlateLayoutTransform(bScaleComments ? 1.0f : Inverse(CurWidget.Geometry.Scale)));
 							PaintComment(NodeComment, CommentGeometry, MyClippingRect, OutDrawElements, ChildLayerId, ChildNode->GetNodeCommentColor().GetColor( InWidgetStyle ), /*inout*/ CommentBubbleY, InWidgetStyle);
 						}
 					}
diff --git a/Engine/Source/Editor/KismetWidgets/Private/SScrubWidget.cpp b/Engine/Source/Editor/KismetWidgets/Private/SScrubWidget.cpp
--- a/Engine/Source/Editor/KismetWidgets/Private/SScrubWidget.cpp
+++ b/Engine/Source/Editor/KismetWidgets/Private/SScrubWidget.cpp
@@ -75,7 +75,7 @@ int32 SScrubWidget::OnPaint( const FPaintArgs& Args, const FGeometry& AllottedGe
 	const int32 TextLayer = BackgroundLayer + 1;
 
 	const FSlateBrush* StyleInfo = FEditorStyle::GetBrush( TEXT( "ProgressBar.Background" ) );
-	const FSlateRect GeomRect = AllottedGeometry.GetRect();
+	const float GeomHeight = AllottedGeometry.Size.Y;
 
 	if ( NumOfKeys.Get() > 0 && SequenceLength.Get() > 0)
 	{
@@ -95,7 +95,7 @@ int32 SScrubWidget::OnPaint( const FPaintArgs& Args, const FGeometry& AllottedGe
 			if ( FGenericPlatformMath::Fmod(KeyVal, Divider) == 0.f )
 			{
 				const FVector2D Offset(XPos, 0.f);
-				const FVector2D Size(1, GeomRect.Bottom-GeomRect.Top);
+				const FVector2D Size(1, GeomHeight);
 				// draw each box with key frame
 				FSlateDrawElement::MakeBox(
 					OutDrawElements,
@@ -125,7 +125,7 @@ int32 SScrubWidget::OnPaint( const FPaintArgs& Args, const FGeometry& AllottedGe
 			}
  			else if (HalfDivider > 1.f)
  			{
- 				const float Height = GeomRect.Bottom-GeomRect.Top;
+ 				const float Height = GeomHeight;
  				const FVector2D Offset(XPos, Height*0.25f);
  				const FVector2D Size(1, Height*0.5f);
  				// draw each box with key frame
@@ -171,7 +171,7 @@ int32 SScrubWidget::OnPaint( const FPaintArgs& Args, const FGeometry& AllottedGe
 			{
 				const float BarXPos = TimeScaleInfo.InputToLocalX(BarValue);
 				const FVector2D BarOffset(BarXPos-2.f, 0.f);
-				const FVector2D Size(4.f, GeomRect.Bottom-GeomRect.Top);
+				const FVector2D Size(4.f, GeomHeight);
 
 				FLinearColor BarColor = InWidgetStyle.GetColorAndOpacityTint();	
 				BarColor.R *= 0.1f;
diff --git a/Engine/Source/Editor/Sequencer/Private/TimeSliderController.cpp b/Engine/Source/Editor/Sequencer/Private/TimeSliderController.cpp
--- a/Engine/Source/Editor/Sequencer/Private/TimeSliderController.cpp
+++ b/Engine/Source/Editor/Sequencer/Private/TimeSliderController.cpp
@@ -229,8 +229,6 @@ int32 FSequencerTimeSliderController::OnPaintTimeSlider( bool bMirrorLabels, con
 	const bool bEnabled = bParentEnabled;
 	const ESlateDrawEffect::Type DrawEffects = bEnabled ? ESlateDrawEffect::None : ESlateDrawEffect::DisabledEffect;
 
-	FSlateRect GeomRect = AllottedGeometry.GetRect();
-
 	TRange<float> LocalViewRange = TimeSliderArgs.ViewRange.Get();
 	const float LocalViewRangeMin = LocalViewRange.GetLowerBoundValue();
 	const float LocalViewRangeMax = LocalViewRange.GetUpperBoundValue();
diff --git a/Engine/Source/Editor/UMGEditor/Private/Designer/DesignTimeUtils.cpp b/Engine/Source/Editor/UMGEditor/Private/Designer/DesignTimeUtils.cpp
--- a/Engine/Source/Editor/UMGEditor/Private/Designer/DesignTimeUtils.cpp
+++ b/Engine/Source/Editor/UMGEditor/Private/Designer/DesignTimeUtils.cpp
@@ -38,7 +38,7 @@ bool FDesignTimeUtils::GetArrangedWidgetRelativeToWindow(TSharedRef<SWidget> Wid
 	if ( FSlateApplication::Get().GeneratePathToWidgetUnchecked(Widget, WidgetPath) )
 	{
 		ArrangedWidget = WidgetPath.FindArrangedWidget(Widget);
-		ArrangedWidget.Geometry.AbsolutePosition -= CurrentWindowRef->GetPositionInScreen();
+		ArrangedWidget.Geometry.AppendTransform(FSlateLayoutTransform(Inverse(CurrentWindowRef->GetPositionInScreen())));
 		return true;
 	}
 
@@ -50,10 +50,7 @@ bool FDesignTimeUtils::GetArrangedWidgetRelativeToParent(TSharedRef<const SWidge
 	FWidgetPath WidgetPath;
 	if ( FSlateApplication::Get().GeneratePathToWidgetUnchecked(Widget, WidgetPath) )
 	{
-		FArrangedWidget ArrangedDesigner = WidgetPath.FindArrangedWidget(Parent);
-
-		ArrangedWidget = WidgetPath.FindArrangedWidget(Widget);
-		ArrangedWidget.Geometry.AbsolutePosition -= ArrangedDesigner.Geometry.AbsolutePosition;
+		GetArrangedWidgetRelativeToParent(WidgetPath, Widget, Parent, ArrangedWidget);
 		return true;
 	}
 
@@ -65,7 +62,8 @@ void FDesignTimeUtils::GetArrangedWidgetRelativeToParent(FWidgetPath& WidgetPath
 	FArrangedWidget ArrangedDesigner = WidgetPath.FindArrangedWidget(Parent);
 
 	ArrangedWidget = WidgetPath.FindArrangedWidget(Widget);
-	ArrangedWidget.Geometry.AbsolutePosition -= ArrangedDesigner.Geometry.AbsolutePosition;
+	// !!! WRH 2014/08/26 - This code tries to mutate an FGeometry in an unsupported way. This code should be refactored.
+	ArrangedWidget.Geometry.AppendTransform(FSlateLayoutTransform(Inverse(ArrangedDesigner.Geometry.AbsolutePosition)));
 }
 
 #undef LOCTEXT_NAMESPACE
diff --git a/Engine/Source/Editor/UMGEditor/Private/Designer/SDesignerView.cpp b/Engine/Source/Editor/UMGEditor/Private/Designer/SDesignerView.cpp
--- a/Engine/Source/Editor/UMGEditor/Private/Designer/SDesignerView.cpp
+++ b/Engine/Source/Editor/UMGEditor/Private/Designer/SDesignerView.cpp
@@ -861,10 +861,8 @@ int32 SDesignerView::OnPaint(const FPaintArgs& Args, const FGeometry& AllottedGe
 		const FLinearColor Tint(0, 1, 0);
 
 		// Draw selection effect
-		FPaintGeometry SelectionGeometry(
-			ArrangedWidget.Geometry.AbsolutePosition,
-			ArrangedWidget.Geometry.Size * ArrangedWidget.Geometry.Scale,
-			ArrangedWidget.Geometry.Scale);
+		
+		FPaintGeometry SelectionGeometry = ArrangedWidget.Geometry.ToPaintGeometry();
 
 		FSlateDrawElement::MakeBox(
 			OutDrawElements,
@@ -876,20 +874,20 @@ int32 SDesignerView::OnPaint(const FPaintArgs& Args, const FGeometry& AllottedGe
 			Tint
 		);
 
-		DrawDragHandles(SelectionGeometry, MyClippingRect, OutDrawElements, LayerId, InWidgetStyle);
+		DrawDragHandles(ArrangedWidget.Geometry, MyClippingRect, OutDrawElements, LayerId, InWidgetStyle);
 	}
 
 	return LayerId;
 }
 
-void SDesignerView::DrawDragHandles(const FPaintGeometry& SelectionGeometry, const FSlateRect& MyClippingRect, FSlateWindowElementList& OutDrawElements, int32 LayerId, const FWidgetStyle& InWidgetStyle) const
+void SDesignerView::DrawDragHandles(const FGeometry& SelectionGeometry, const FSlateRect& MyClippingRect, FSlateWindowElementList& OutDrawElements, int32 LayerId, const FWidgetStyle& InWidgetStyle) const
 {
 	if ( SelectedWidget.IsValid() && SelectedWidget.GetTemplate()->Slot )
 	{
-		float X = SelectionGeometry.DrawPosition.X;
-		float Y = SelectionGeometry.DrawPosition.Y;
-		float Width = SelectionGeometry.DrawSize.X;
-		float Height = SelectionGeometry.DrawSize.Y;
+		float X = 0.0f;
+		float Y = 0.0f;
+		float Width = SelectionGeometry.Size.X;
+		float Height = SelectionGeometry.Size.Y;
 
 		// @TODO UMG Handles should come from the slot/container to express how its slots can be transformed.
 		TArray<FVector2D> Handles;
@@ -922,7 +920,7 @@ void SDesignerView::DrawDragHandles(const FPaintGeometry& SelectionGeometry, con
 			FSlateDrawElement::MakeBox(
 				OutDrawElements,
 				++LayerId,
-				FPaintGeometry(FVector2D(Handle.X - HandleSize.X * 0.5f, Handle.Y - HandleSize.Y * 0.5f), HandleSize, 1.0f),
+				SelectionGeometry.ToPaintGeometry(HandleSize, FSlateLayoutTransform(FVector2D(Handle.X - HandleSize.X * 0.5f, Handle.Y - HandleSize.Y * 0.5f))),
 				KeyBrush,
 				MyClippingRect,
 				ESlateDrawEffect::None,
diff --git a/Engine/Source/Editor/UMGEditor/Private/Designer/SDesignerView.h b/Engine/Source/Editor/UMGEditor/Private/Designer/SDesignerView.h
--- a/Engine/Source/Editor/UMGEditor/Private/Designer/SDesignerView.h
+++ b/Engine/Source/Editor/UMGEditor/Private/Designer/SDesignerView.h
@@ -150,7 +150,7 @@ private:
 	TArray< TSharedRef<FDesignerExtension> > DesignerExtensions;
 
 private:
-	void DrawDragHandles(const FPaintGeometry& SelectionGeometry, const FSlateRect& MyClippingRect, FSlateWindowElementList& OutDrawElements, int32 LayerId, const FWidgetStyle& InWidgetStyle) const;
+	void DrawDragHandles(const FGeometry& SelectionGeometry, const FSlateRect& MyClippingRect, FSlateWindowElementList& OutDrawElements, int32 LayerId, const FWidgetStyle& InWidgetStyle) const;
 	DragHandle HitTestDragHandles(const FGeometry& AllottedGeometry, const FPointerEvent& PointerEvent) const;
 
 	UWidget* ProcessDropAndAddWidget(const FGeometry& MyGeometry, const FDragDropEvent& DragDropEvent, bool bIsPreview);
diff --git a/Engine/Source/Editor/UnrealEd/Private/PlayLevel.cpp b/Engine/Source/Editor/UnrealEd/Private/PlayLevel.cpp
--- a/Engine/Source/Editor/UnrealEd/Private/PlayLevel.cpp
+++ b/Engine/Source/Editor/UnrealEd/Private/PlayLevel.cpp
@@ -2431,7 +2431,6 @@ UGameInstance* UEditorEngine::CreatePIEGameInstance(int32 PIEInstance, bool bInS
 						{
 							// Save off the window position
 							const FVector2D PIEWindowPos = WindowBeingClosed->GetPositionInScreen();
-							const FVector2D PIEWindowClientSize = WindowBeingClosed->GetClientSizeInScreen();
 
 							ULevelEditorPlaySettings* LevelEditorPlaySettings = ULevelEditorPlaySettings::StaticClass()->GetDefaultObject<ULevelEditorPlaySettings>();
 
diff --git a/Engine/Source/Editor/UnrealEd/Private/SColorGradientEditor.cpp b/Engine/Source/Editor/UnrealEd/Private/SColorGradientEditor.cpp
--- a/Engine/Source/Editor/UnrealEd/Private/SColorGradientEditor.cpp
+++ b/Engine/Source/Editor/UnrealEd/Private/SColorGradientEditor.cpp
@@ -252,8 +252,7 @@ int32 SColorGradientEditor::OnPaint( const FPaintArgs& Args, const FGeometry& Al
 			// Draw the text centered in the color region
 			{
 				FVector2D StringSize = FontMeasureService->Measure( GradientColorMessage, FSlateFontInfo( FPaths::EngineContentDir() / TEXT("Slate/Fonts/Roboto-Regular.ttf"), 8 ) );
-				FPaintGeometry PaintGeom = ColorMarkAreaGeometry.CenteredPaintGeometryBelow( StringSize, 1.0f );
-				PaintGeom.DrawPosition.Y = ColorMarkAreaGeometry.AbsolutePosition.Y + 1;
+				FPaintGeometry PaintGeom = ColorMarkAreaGeometry.ToPaintGeometry(FSlateLayoutTransform(FVector2D((ColorMarkAreaGeometry.Size.X - StringSize.X) * 0.5f, 1.0f)));
 
 				FSlateDrawElement::MakeText
 				( 
@@ -271,8 +270,7 @@ int32 SColorGradientEditor::OnPaint( const FPaintArgs& Args, const FGeometry& Al
 			// Draw the text centered in the alpha region
 			{
 				FVector2D StringSize = FontMeasureService->Measure( GradientAlphaMessage, FSlateFontInfo( FPaths::EngineContentDir() / TEXT("Slate/Fonts/Roboto-Regular.ttf"), 8 ) );
-				FPaintGeometry PaintGeom = AlphaMarkAreaGeometry.CenteredPaintGeometryBelow( StringSize, 1.0f );
-				PaintGeom.DrawPosition.Y = AlphaMarkAreaGeometry.AbsolutePosition.Y + 1;
+				FPaintGeometry PaintGeom = AlphaMarkAreaGeometry.ToPaintGeometry(FSlateLayoutTransform(FVector2D((AlphaMarkAreaGeometry.Size.X - StringSize.X) * 0.5f, 1.0f)));
 
 				FSlateDrawElement::MakeText
 				( 
diff --git a/Engine/Source/Editor/WorldBrowser/Private/Tiles/SWorldComposition.cpp b/Engine/Source/Editor/WorldBrowser/Private/Tiles/SWorldComposition.cpp
--- a/Engine/Source/Editor/WorldBrowser/Private/Tiles/SWorldComposition.cpp
+++ b/Engine/Source/Editor/WorldBrowser/Private/Tiles/SWorldComposition.cpp
@@ -478,7 +478,7 @@ protected:
 	uint32 PaintBackground(const FGeometry& AllottedGeometry, const FSlateRect& MyClippingRect, FSlateWindowElementList& OutDrawElements, uint32 LayerId) const
 	{
 		FVector2D ScreenWorldOrigin = GraphCoordToPanelCoord(FVector2D(0, 0));
-		FSlateRect ScreenRect = AllottedGeometry.GetRect();
+		FSlateRect ScreenRect(FVector2D(0, 0), AllottedGeometry.GetLocalSize());
 	
 		// World Y-axis
 		if (ScreenWorldOrigin.X > ScreenRect.Left &&
@@ -643,7 +643,7 @@ protected:
 	{
 		VisibleChildren.Empty();
 
-		FSlateRect PanelRect = AllottedGeometry.GetRect();
+		FSlateRect PanelRect(FVector2D(0, 0), AllottedGeometry.GetLocalSize());
 		FVector2D ViewStartPos = PanelCoordToGraphCoord(FVector2D(PanelRect.Left, PanelRect.Top));
 		FVector2D ViewEndPos = PanelCoordToGraphCoord(FVector2D(PanelRect.Right, PanelRect.Bottom));
 		FSlateRect ViewRect(ViewStartPos, ViewEndPos);
diff --git a/Engine/Source/Editor/WorldBrowser/Private/Tiles/SWorldTileItem.cpp b/Engine/Source/Editor/WorldBrowser/Private/Tiles/SWorldTileItem.cpp
--- a/Engine/Source/Editor/WorldBrowser/Private/Tiles/SWorldTileItem.cpp
+++ b/Engine/Source/Editor/WorldBrowser/Private/Tiles/SWorldTileItem.cpp
@@ -414,7 +414,7 @@ int32 SWorldTileItem::OnPaint(const FPaintArgs& Args, const FGeometry& AllottedG
 			const float ProgressBarAnimOffset = ProgressBarImage->ImageSize.X * CurveSequence.GetLerpLooping() / AllottedGeometry.Scale;
 			const float ProgressBarImageSize = ProgressBarImage->ImageSize.X / AllottedGeometry.Scale;
 			const float ProgressBarImageHeight = ProgressBarImage->ImageSize.Y / AllottedGeometry.Scale;
-			const FSlateRect ProgressBarClippingRect = AllottedGeometry.ToPaintGeometry().ToSlateRect().IntersectionWith(ClippingRect);
+			const FSlateRect ProgressBarClippingRect = AllottedGeometry.GetClippingRect().IntersectionWith(ClippingRect);
 		
 			FPaintGeometry LoadingBarPaintGeometry = AllottedGeometry.ToPaintGeometry(
 					FVector2D(ProgressBarAnimOffset - ProgressBarImageSize, 0 ),
diff --git a/Engine/Source/Runtime/AppFramework/Private/Framework/Testing/SLayoutExample.cpp b/Engine/Source/Runtime/AppFramework/Private/Framework/Testing/SLayoutExample.cpp
--- a/Engine/Source/Runtime/AppFramework/Private/Framework/Testing/SLayoutExample.cpp
+++ b/Engine/Source/Runtime/AppFramework/Private/Framework/Testing/SLayoutExample.cpp
@@ -397,7 +397,10 @@ public:
 
 APPFRAMEWORK_API TSharedRef<SWidget> MakeLayoutExample()
 {
-	return SNew(SExampleLayout);
+	extern FSlateRenderTransform GetRenderTransform();
+	return 
+		SNew(SExampleLayout)
+		.RenderTransform_Static(&GetRenderTransform);
 }
 
 #undef LOCTEXT_NAMESPACE
\ No newline at end of file
diff --git a/Engine/Source/Runtime/AppFramework/Private/Framework/Testing/STableViewTesting.cpp b/Engine/Source/Runtime/AppFramework/Private/Framework/Testing/STableViewTesting.cpp
--- a/Engine/Source/Runtime/AppFramework/Private/Framework/Testing/STableViewTesting.cpp
+++ b/Engine/Source/Runtime/AppFramework/Private/Framework/Testing/STableViewTesting.cpp
@@ -1102,7 +1102,10 @@ protected:
 
 TSharedRef<SWidget> MakeTableViewTesting()
 {
-	return SNew( STableViewTesting );
+	extern FSlateRenderTransform GetRenderTransform();
+	return 
+		SNew( STableViewTesting )
+		.RenderTransform_Static(&GetRenderTransform);
 }
 
 
diff --git a/Engine/Source/Runtime/AppFramework/Private/Framework/Testing/SWidgetGallery.cpp b/Engine/Source/Runtime/AppFramework/Private/Framework/Testing/SWidgetGallery.cpp
--- a/Engine/Source/Runtime/AppFramework/Private/Framework/Testing/SWidgetGallery.cpp
+++ b/Engine/Source/Runtime/AppFramework/Private/Framework/Testing/SWidgetGallery.cpp
@@ -249,6 +249,8 @@ public:
 						.HAlign(HAlign_Left)
 						.Padding(0.0f, 5.0f)
 						[
+							SNew(SFxWidget).IgnoreClipping(false)
+							[
 							SNew(SVerticalBox)
 
 							+ SVerticalBox::Slot()
@@ -256,6 +258,7 @@ public:
 								[
 									// selector combo box
 									SAssignNew(SelectorComboBox, SComboBox<TSharedPtr<FString> >)
+										.Method(SMenuAnchor::UseCurrentWindow)
 										.OptionsSource(&SelectorComboBoxOptions)
 										.OnSelectionChanged(this, &SWidgetGallery::HandleSelectorComboBoxSelectionChanged)
 										.OnGenerateWidget(this, &SWidgetGallery::HandleComboBoxGenerateWidget)
@@ -278,6 +281,7 @@ public:
 												.Text(this, &SWidgetGallery::HandleSecondComboBoxText)
 										]
 								]
+							]
 						]
 
 					// SComboButton
@@ -291,7 +295,9 @@ public:
 						.HAlign(HAlign_Left)
 						.Padding(0.0f, 5.0f)
 						[
-							SNew(SComboButton)
+							SNew(SFxWidget).IgnoreClipping(false)
+							[
+								SNew(SComboButton)
 								.Method(SMenuAnchor::UseCurrentWindow)
 								.ButtonContent()
 								[
@@ -327,6 +333,7 @@ public:
 												.Text(LOCTEXT("ComboButtonItemLabel03", "widgets"))
 										]
 								]
+							]
 						]
 
 					// SEditableText
@@ -1093,7 +1100,10 @@ private:
  */
 TSharedRef<SWidget> MakeWidgetGallery()
 {
-	return SNew(SWidgetGallery);
+	extern FSlateRenderTransform GetRenderTransform();
+	return 
+		SNew(SWidgetGallery)
+		.RenderTransform_Static(&GetRenderTransform);
 }
 
 
diff --git a/Engine/Source/Runtime/AppFramework/Private/Widgets/Testing/STestSuite.cpp b/Engine/Source/Runtime/AppFramework/Private/Widgets/Testing/STestSuite.cpp
--- a/Engine/Source/Runtime/AppFramework/Private/Widgets/Testing/STestSuite.cpp
+++ b/Engine/Source/Runtime/AppFramework/Private/Widgets/Testing/STestSuite.cpp
@@ -7,10 +7,195 @@
 #include "TestStyle.h"
 #include "SyntaxHighlighterTextLayoutMarshaller.h"
 #include "STestSuite.h"
-
+#include "TransformCalculus3D.h"
+#include "SlateRenderTransform.h"
+#include "SlateLayoutTransform.h"
 
 #define LOCTEXT_NAMESPACE "STestSuite"
 
+namespace
+{
+	/**
+	 * A set of syntax tests for transform calculus to ensure that the basic primitives are supported properly.
+	 * Does NOT test for correctness, purely syntax!
+	 */
+	bool TestTransformCalculus()
+	{
+		float uniScale = 5.8f;
+		FVector trans(5, 6, 7);
+		FQuat quat(FVector(1, 2, 3).SafeNormal(), 33.5f);
+		FRotator rot(7, 8, 9);
+		FMatrix mat = FRotationMatrix::Make(rot);
+
+		// Identity casts
+		const float& uniScale2 = TransformCast<float>(uniScale);
+		const FVector& trans2 = TransformCast<FVector>(trans);
+		const FQuat& quat2 = TransformCast<FQuat>(quat);
+		const FRotator& rot2 = TransformCast<FRotator>(rot);
+		const FMatrix& mat2 = TransformCast<FMatrix>(mat);
+
+		// rotation casts
+		auto quat3 = TransformCast<FQuat>(rot);
+		auto rot3 = TransformCast<FRotator>(quat);
+		rot3 = TransformCast<FRotator>(mat);
+
+		// higher level transform casts
+		auto mat3 = TransformCast<FMatrix>(uniScale);
+		mat3 = TransformCast<FMatrix>(trans);
+		mat3 = TransformCast<FMatrix>(quat);
+		mat3 = TransformCast<FMatrix>(rot);
+
+		// basic concatenation
+		auto uniScale4 = Concatenate(uniScale, uniScale2);
+		auto trans4 = Concatenate(trans, trans2);
+		auto quat4 = Concatenate(quat, quat2);
+		auto rot4 = Concatenate(rot, rot2);
+		auto mat4 = Concatenate(mat, mat2);
+
+		// matrix conversion
+		FMatrix mat5;
+		FQuat quat5;
+
+		mat5 = Concatenate(uniScale4, trans4);
+		mat5 = Concatenate(trans4, uniScale4);
+		mat5 = Concatenate(uniScale4, quat4);
+		mat5 = Concatenate(quat4, uniScale4);
+		mat5 = Concatenate(uniScale4, rot4);
+		mat5 = Concatenate(rot4, uniScale4);
+		mat5 = Concatenate(uniScale4, mat4);
+		mat5 = Concatenate(mat4, uniScale4);
+
+		mat5 = Concatenate(trans4, quat4);
+		mat5 = Concatenate(quat4, trans4);
+		mat5 = Concatenate(trans4, rot4);
+		mat5 = Concatenate(rot4, trans4);
+		mat5 = Concatenate(trans4, mat4);
+		mat5 = Concatenate(mat4, trans4);
+
+		quat5 = Concatenate(quat4, rot4);
+		quat5 = Concatenate(rot4, quat4);
+		mat5 = Concatenate(quat4, mat4);
+		mat5 = Concatenate(mat4, quat4);
+
+		mat5 = Concatenate(rot4, mat4);
+		mat5 = Concatenate(mat4, rot4);
+
+		mat5 = Concatenate(Inverse(trans), rot, trans);
+
+		mat5 = Concatenate(Inverse(trans), rot, uniScale, trans);
+
+		FScale2D scale2D(1.4f, 3.2f);
+		FShear2D shear2D(-3.5f, -4.6f);
+		FVector2D trans2D(7, 8);
+		FQuat2D rot2D(1.4f);
+		FMatrix2x2 mat2D(5, 6, 7, 8);
+		FSlateLayoutTransform transformLayout(uniScale, trans2D);
+		FSlateRenderTransform transform2D = ::Concatenate(FScale2D(2, 4), FShear2D(0.5, 2), FQuat2D(FMath::DegreesToRadians(45.0f)), FVector2D(5, 7));
+
+		// Identity casts
+		const FScale2D& scale2D2 = TransformCast<FScale2D>(scale2D);
+		const FShear2D& shear2D2 = TransformCast<FShear2D>(shear2D);
+		const FVector2D& trans2D2 = TransformCast<FVector2D>(trans2D);
+		const FQuat2D& rot2D2 = TransformCast<FQuat2D>(rot2D);
+		const FMatrix2x2& mat2D2 = TransformCast<FMatrix2x2>(mat2D);
+		const FSlateLayoutTransform& transformLayout2 = TransformCast<FSlateLayoutTransform>(transformLayout);
+		const FSlateRenderTransform& transform2D2 = TransformCast<FSlateRenderTransform>(transform2D);
+
+		auto scale2D3 = Concatenate(scale2D, scale2D);
+		auto shear2D3 = Concatenate(shear2D, shear2D);
+		auto trans2D3 = Concatenate(trans2D, trans2D);
+		auto rot2D3 = Concatenate(rot2D, rot2D);
+		auto mat2D3 = Concatenate(mat2D, mat2D);
+		auto transformLayout3 = Concatenate(transformLayout, transformLayout);
+		auto transform2D3 = Concatenate(transform2D, transform2D);
+
+		// higher level transform casts
+		transformLayout3 = TransformCast<FSlateLayoutTransform>(uniScale);
+		transformLayout3 = TransformCast<FSlateLayoutTransform>(trans2D);
+
+		transform2D3 = TransformCast<FSlateRenderTransform>(uniScale);
+		transform2D3 = TransformCast<FSlateRenderTransform>(scale2D);
+		transform2D3 = TransformCast<FSlateRenderTransform>(shear2D);
+		transform2D3 = TransformCast<FSlateRenderTransform>(trans2D);
+		transform2D3 = TransformCast<FSlateRenderTransform>(rot2D);
+		transform2D3 = TransformCast<FSlateRenderTransform>(mat2D);
+		transform2D3 = TransformCast<FSlateRenderTransform>(transformLayout);
+
+		// basic concatenation
+		scale2D3 = Concatenate(uniScale, scale2D);
+		scale2D3 = Concatenate(scale2D, uniScale);
+
+		auto transformLayout4 = Concatenate(transformLayout, transformLayout);
+		transformLayout4 = Concatenate(uniScale, trans2D);
+		transformLayout4 = Concatenate(trans2D, uniScale);
+
+		transformLayout4 = Concatenate(uniScale, transformLayout);
+		transformLayout4 = Concatenate(trans2D, transformLayout);
+		transformLayout4 = Concatenate(transformLayout, uniScale);
+		transformLayout4 = Concatenate(transformLayout, trans2D);
+
+		auto transform2D4 = Concatenate(transform2D, transform2D);
+
+		transform2D4 = Concatenate(trans2D, scale2D);
+		transform2D4 = Concatenate(trans2D, shear2D);
+		transform2D4 = Concatenate(trans2D, rot2D);
+		transform2D4 = Concatenate(trans2D, mat2D);
+		transform2D4 = Concatenate(trans2D, transform2D);
+		transform2D4 = Concatenate(scale2D, transformLayout);
+		transform2D4 = Concatenate(scale2D, transform2D);
+		transform2D4 = Concatenate(shear2D, transformLayout);
+		transform2D4 = Concatenate(shear2D, transform2D);
+		transform2D4 = Concatenate(rot2D, transformLayout);
+		transform2D4 = Concatenate(rot2D, transform2D);
+		transform2D4 = Concatenate(mat2D, transformLayout);
+		transform2D4 = Concatenate(mat2D, transform2D);
+		transform2D4 = Concatenate(transformLayout, transform2D);
+
+		transform2D4 = Concatenate(scale2D, trans2D);
+		transform2D4 = Concatenate(shear2D, trans2D);
+		transform2D4 = Concatenate(rot2D, trans2D);
+		transform2D4 = Concatenate(mat2D, trans2D);
+		transform2D4 = Concatenate(transform2D, trans2D);
+		transform2D4 = Concatenate(transformLayout, scale2D);
+		transform2D4 = Concatenate(transform2D, scale2D);
+		transform2D4 = Concatenate(transformLayout, shear2D);
+		transform2D4 = Concatenate(transform2D, shear2D);
+		transform2D4 = Concatenate(transformLayout, rot2D);
+		transform2D4 = Concatenate(transform2D, rot2D);
+		transform2D4 = Concatenate(transformLayout, mat2D);
+		transform2D4 = Concatenate(transform2D, mat2D);
+		transform2D4 = Concatenate(transform2D, transformLayout);
+
+		auto mat2D4 = Concatenate(scale2D, shear2D);
+		mat2D4 = Concatenate(scale2D, rot2D);
+		mat2D4 = Concatenate(scale2D, mat2D);
+
+		mat2D4 = Concatenate(shear2D, scale2D);
+		mat2D4 = Concatenate(rot2D, scale2D);
+		mat2D4 = Concatenate(mat2D, scale2D);
+
+		mat2D4 = Concatenate(shear2D, rot2D);
+		mat2D4 = Concatenate(shear2D, mat2D);
+
+		mat2D4 = Concatenate(rot2D, shear2D);
+		mat2D4 = Concatenate(mat2D, shear2D);
+
+		mat2D4 = Concatenate(rot2D, mat2D);
+
+		mat2D4 = Concatenate(mat2D, rot2D);
+		{
+			auto matRot = Concatenate(FRotationMatrix::Make(FRotator(0.0f, 17.5f, 0.0f)), FRotationMatrix::Make(FRotator(12.4f, 5.7f, 29.0f)));
+			auto quatRot = Concatenate(FQuat(FRotator(0.0f, 17.5f, 0.0f)), FQuat(FRotator(12.4f, 5.7f, 29.0f)));
+
+			FVector vec(3.0f, 4.0f, 5.0f);
+			auto matRotRes = TransformPoint(matRot, vec);
+			auto quatRotRes = TransformPoint(quatRot, vec);
+		}
+
+		return true;
+	}
+	bool bTestTransformCalculus = TestTransformCalculus();
+}
 
 struct FOnPaintHandlerParams
 {
@@ -349,10 +534,10 @@ private:
 	{
 		TArray<FSlateGradientStop> GradientStops;
 
-		GradientStops.Add( FSlateGradientStop( FVector2D::ZeroVector, FColor(255,255,0) ) );
+		GradientStops.Add( FSlateGradientStop( FVector2D(InParams.Geometry.Size.X*.1f,0), FColor(255,255,0) ) );
 		GradientStops.Add( FSlateGradientStop( FVector2D(InParams.Geometry.Size.X*.25f,0), FColor(255,0,255) ) );
 		GradientStops.Add( FSlateGradientStop( FVector2D(InParams.Geometry.Size.X*.75f,0), FColor(0,0,255) ) );
-		GradientStops.Add( FSlateGradientStop( InParams.Geometry.Size, FColor(0,255,0) ) );
+		GradientStops.Add( FSlateGradientStop( FVector2D(InParams.Geometry.Size.X*0.9f,0), FColor(0,255,0) ) );
 
 		FSlateDrawElement::MakeGradient(
 			InParams.OutDrawElements,
@@ -370,9 +555,9 @@ private:
 	int32 TestSplineElement( const FOnPaintHandlerParams& InParams )
 	{
 		const FVector2D Start(10,10);
-		const FVector2D StartDir(1000,0);
+		const FVector2D StartDir(InParams.Geometry.Size.X * 1000 / 600,0);
 		const FVector2D End(InParams.Geometry.Size.X/4, InParams.Geometry.Size.Y-10);
-		const FVector2D EndDir(1000,0);
+		const FVector2D EndDir(InParams.Geometry.Size.X * 1000 / 600,0);
 
 		FSlateDrawElement::MakeSpline(
 			InParams.OutDrawElements,
@@ -432,12 +617,12 @@ private:
 	void MakeRotationExample( const FOnPaintHandlerParams& InParams )
 	{
 		const FSlateBrush* CenterBrush = FTestStyle::Get().GetBrush("TestRotation40px");
+		const FSlateBrush* TestBrush = FTestStyle::Get().GetBrush("TestRotation20px");
 
-		FVector2D LocalPos = FVector2D(50,50);
-		FVector2D Size = CenterBrush->ImageSize;
-		FVector2D OrbitPos = (LocalPos + (LocalPos+Size))*.5f;
+		const FVector2D LocalPos = FVector2D(50,50);
+		const FVector2D LocalSize = CenterBrush->ImageSize;
 
-		FPaintGeometry CenterGeom = InParams.Geometry.ToPaintGeometry( LocalPos, Size );
+		FSlateLayoutTransform CenterLayoutTransform(LocalPos);
 
 		// Make a box that rotates around its center.  Note if you don't specify the rotation point or rotation space
 		// it defaults to rotating about the center of the box.  ERotationSpace::RelativeToElement is used by default in this case.  
@@ -446,7 +631,7 @@ private:
 			FSlateDrawElement::MakeRotatedBox(
 				InParams.OutDrawElements,
 				InParams.Layer,
-				CenterGeom,
+				InParams.Geometry.ToPaintGeometry( LocalSize, CenterLayoutTransform ),
 				CenterBrush,
 				InParams.ClippingRect,
 				InParams.bEnabled ? ESlateDrawEffect::None : ESlateDrawEffect::DisabledEffect,
@@ -455,23 +640,20 @@ private:
 		}
 
 		// Make a box that rotates around the center of the previous box
-		// In this example we rotate around a point( the center of the previous box ), which is in world space so we specify ERotationSpace::RelativeToWorld
 		{
-			const FSlateBrush* TestBrush = FTestStyle::Get().GetBrush("TestRotation20px");
-
-			FVector2D WorldOrbitPos = InParams.Geometry.LocalToAbsolute( OrbitPos );
+			FSlateLayoutTransform OrbitLayoutTransform(LocalPos + LocalSize);
+			const FVector2D LocalCenterOfRotation = (LocalPos + (LocalPos+LocalSize))*.5f;
+			const FVector2D RelativeOrbitPos = TransformPoint(Inverse(OrbitLayoutTransform), LocalCenterOfRotation);
 
 			FSlateDrawElement::MakeRotatedBox(
 				InParams.OutDrawElements,
 				InParams.Layer,
-				InParams.Geometry.ToPaintGeometry( FVector2D(110,50), TestBrush->ImageSize ),
+				InParams.Geometry.ToPaintGeometry( TestBrush->ImageSize, OrbitLayoutTransform ),
 				TestBrush,
 				InParams.ClippingRect,
 				InParams.bEnabled ? ESlateDrawEffect::None : ESlateDrawEffect::DisabledEffect,
 				OuterRotation,
-				WorldOrbitPos,
-				FSlateDrawElement::RelativeToWorld
-				);
+				RelativeOrbitPos);
 		}
 
 	}
@@ -3268,6 +3450,7 @@ class SFxTest : public SCompoundWidget
 		RenderScaleOrigin = FVector2D(0.5f, 0.5f);
 		LayoutScale = 1.0f;
 		VisualOffset = FVector2D::ZeroVector;
+		FxWidgetIgnoreClippingState = ESlateCheckBoxState::Checked;
 
 
 		this->ChildSlot
@@ -3350,6 +3533,16 @@ class SFxTest : public SCompoundWidget
 							.OnValueChanged(this, &SFxTest::OnVisualOffsetChangedY)
 						]
 					]
+					+SGridPanel::Slot(0, 4) .Padding(2)
+					[
+						SNew(SCheckBox)
+						.IsChecked( this, &SFxTest::GetFxWidgetIgnoreClippingState )
+						.OnCheckStateChanged( this, &SFxTest::OnFxWidgetIgnoreClippingChanged )
+						[
+							SNew(STextBlock)
+							.Text( LOCTEXT("IgnoreClipping", "Ignore Clipping") )
+						]
+					]
 				]
 				+SHorizontalBox::Slot()
 				[
@@ -3383,6 +3576,7 @@ class SFxTest : public SCompoundWidget
 				SNew(SBorder)
 				[
 					SNew(SFxWidget)
+					.IgnoreClipping(this, &SFxTest::GetFxWidgetIgnoreClipping)
 					.RenderScale( this, &SFxTest::GetRenderScale )
 					.RenderScaleOrigin( this, &SFxTest::GetRenderScaleOrigin )
 					.LayoutScale( this, &SFxTest::GetLayoutScale )
@@ -3409,6 +3603,10 @@ class SFxTest : public SCompoundWidget
 	}
 	END_SLATE_FUNCTION_BUILD_OPTIMIZATION
 
+	bool GetFxWidgetIgnoreClipping() const { return FxWidgetIgnoreClippingState == ESlateCheckBoxState::Checked; }
+	ESlateCheckBoxState::Type GetFxWidgetIgnoreClippingState() const { return FxWidgetIgnoreClippingState; }
+	void OnFxWidgetIgnoreClippingChanged( ESlateCheckBoxState::Type InValue ) { FxWidgetIgnoreClippingState = InValue; }
+	ESlateCheckBoxState::Type FxWidgetIgnoreClippingState;
 
 	float GetRenderScale() const { return RenderScale; }
 	void OnRenderScaleChanged( float InValue ) { RenderScale = InValue; }
@@ -4047,11 +4245,202 @@ public:
 	END_SLATE_FUNCTION_BUILD_OPTIMIZATION
 };
 
+/**
+ * User Widget wrapper for a this widget. Legacy as this widget used to be in a separate header.
+ */
+class SRenderTransformManipulatorWidget : public SUserWidget
+{
+
+public:
+
+	SLATE_USER_ARGS(SRenderTransformManipulatorWidget)
+	{}
+
+	SLATE_END_ARGS()
+
+	virtual void Construct(const FArguments& InArgs) = 0;
+};
+
+namespace
+{
+	float RotDeg = 0.0f;
+	FScale2D Scale;
+	FShear2D Shear;
+	FQuat2D Rot;
+	FVector2D Offset(0,0);
+
+	class SRenderTransformManipulatorWidgetImpl : public SRenderTransformManipulatorWidget
+	{
+	public:
+		virtual void Construct(const FArguments& InArgs) override;
+
+		TSharedPtr<SImage> ImageWidget;
+	private:
+		static const ISlateStyle& GetStyle()
+		{
+			static FSlateStyleSet Style("RenderTransformManipulatorStyle");
+			static bool IsInit = false;
+			if (!IsInit)
+			{
+				check(IsInGameThread());
+				Style.SetContentRoot(FPaths::EngineContentDir() / TEXT("Slate"));
+				Style.Set("UE4Icon", new FSlateImageBrush(Style.RootToContentDir(TEXT("Testing/UE4Icon.png")), FVector2D(50, 50)));
+				IsInit = true;
+			}
+			return Style;
+		}
+	};
+}
+
+/**
+ * Global access here because we need other translation units to access this function.
+ */
+FSlateRenderTransform GetRenderTransform()
+{
+	return Concatenate(-Offset, Concatenate(Shear, Scale, Rot, Offset));
+}
+
+void SRenderTransformManipulatorWidgetImpl::Construct(const FArguments& InArgs)
+{
+	SUserWidget::Construct(SUserWidget::FArguments()
+		[
+			SNew(SVerticalBox)
+			+ SVerticalBox::Slot().Padding(4)
+				[
+					SAssignNew(ImageWidget, SImage)
+					.RenderTransform_Static(&::GetRenderTransform)
+					.Image(GetStyle().GetBrush("UE4Icon"))
+				]
+			+ SVerticalBox::Slot().AutoHeight().Padding(4)
+			[
+				SNew(SHorizontalBox)
+				+ SHorizontalBox::Slot()
+				[
+					SNew(STextBlock).Text(LOCTEXT("RD", "Rotation degrees"))
+				]
+				+ SHorizontalBox::Slot().HAlign(HAlign_Fill)
+					[
+						SNew(SSpinBox<float>)
+						.MinValue(0.0f)
+						.MaxValue(360.0f)
+						.OnValueChanged_Static([](float val) { RotDeg = val; Rot = FQuat2D(FMath::DegreesToRadians(val)); })
+						.Value_Static([] { return RotDeg; })
+					]
+			]
+			+ SVerticalBox::Slot().AutoHeight().Padding(4)
+				[
+					SNew(SHorizontalBox)
+					+ SHorizontalBox::Slot()
+					[
+						SNew(STextBlock).Text(LOCTEXT("XO", "X Offset"))
+					]
+					+ SHorizontalBox::Slot().HAlign(HAlign_Fill)
+						[
+							SNew(SSpinBox<float>)
+							.MinValue(-1000.0f)
+							.MaxValue(1000.0f)
+							.OnValueChanged_Static([](float val) { Offset.X = val; })
+							.Value_Static([] { return Offset.X; })
+						]
+				]
+			+ SVerticalBox::Slot().AutoHeight().Padding(4)
+				[
+					SNew(SHorizontalBox)
+					+ SHorizontalBox::Slot()
+					[
+						SNew(STextBlock).Text(LOCTEXT("YO", "Y Offset"))
+					]
+					+ SHorizontalBox::Slot().HAlign(HAlign_Fill)
+						[
+							SNew(SSpinBox<float>)
+							.MinValue(-1000.0f)
+							.MaxValue(1000.0f)
+							.OnValueChanged_Static([](float val) { Offset.Y = val; })
+							.Value_Static([] { return Offset.Y; })
+						]
+				]
+			+ SVerticalBox::Slot().AutoHeight().Padding(4)
+				[
+					SNew(SHorizontalBox)
+					+ SHorizontalBox::Slot()
+					[
+						SNew(STextBlock).Text(LOCTEXT("XS", "X Scale"))
+					]
+					+ SHorizontalBox::Slot().HAlign(HAlign_Fill)
+						[
+							SNew(SSpinBox<float>)
+							.MinValue(-10.0f)
+							.MaxValue(10.0f)
+							.OnValueChanged_Static([](float val) { Scale = FScale2D(val, Scale.GetVector().Y); })
+							.Value_Static([] { return Scale.GetVector().X; })
+						]
+				]
+			+ SVerticalBox::Slot().AutoHeight().Padding(4)
+				[
+					SNew(SHorizontalBox)
+					+ SHorizontalBox::Slot()
+					[
+						SNew(STextBlock).Text(LOCTEXT("YS", "Y Scale"))
+					]
+					+ SHorizontalBox::Slot().HAlign(HAlign_Fill)
+						[
+							SNew(SSpinBox<float>)
+							.MinValue(-10.0f)
+							.MaxValue(10.0f)
+							.OnValueChanged_Static([](float val) { Scale = FScale2D(Scale.GetVector().X, val); })
+							.Value_Static([] { return Scale.GetVector().Y; })
+						]
+				]
+			+ SVerticalBox::Slot().AutoHeight().Padding(4)
+				[
+					SNew(SHorizontalBox)
+					+ SHorizontalBox::Slot()
+					[
+						SNew(STextBlock).Text(LOCTEXT("XSH", "X Shear"))
+					]
+					+ SHorizontalBox::Slot().HAlign(HAlign_Fill)
+						[
+							SNew(SSpinBox<float>)
+							.MinValue(-4.f)
+							.MaxValue(4.f)
+							.OnValueChanged_Static([](float val) { Shear = FShear2D(val, Shear.GetVector().Y); })
+							.Value_Static([] { return Shear.GetVector().X; })
+						]
+				]
+			+ SVerticalBox::Slot().AutoHeight().Padding(4)
+				[
+					SNew(SHorizontalBox)
+					+ SHorizontalBox::Slot()
+					[
+						SNew(STextBlock).Text(LOCTEXT("YSH", "Y Shear"))
+					]
+					+ SHorizontalBox::Slot().HAlign(HAlign_Fill)
+						[
+							SNew(SSpinBox<float>)
+							.MinValue(-4.f)
+							.MaxValue(4.f)
+							.OnValueChanged_Static([](float val) { Shear = FShear2D(Shear.GetVector().X, val); })
+							.Value_Static([] { return Shear.GetVector().Y; })
+						]
+				]
+		]
+	);
+}
+
+/**
+ * New implementation for the User Widget idiom.
+ */
+TSharedRef<SRenderTransformManipulatorWidget> SRenderTransformManipulatorWidget::New()
+{
+	return MakeShareable(new SRenderTransformManipulatorWidgetImpl());
+}
+
 static TSharedPtr<FTabManager> TestSuite1TabManager;
 
 BEGIN_SLATE_FUNCTION_BUILD_OPTIMIZATION
 TSharedRef<SDockTab> SpawnTab(const FSpawnTabArgs& Args, FName TabIdentifier)
 {
+	extern FSlateRenderTransform GetRenderTransform();
 	if (TabIdentifier == FName(TEXT("AnimationTestTab")))
 	{
 		return SNew(SDockTab)
@@ -4080,6 +4469,7 @@ TSharedRef<SDockTab> SpawnTab(const FSpawnTabArgs& Args, FName TabIdentifier)
 			)
 		[
 			SNew(SVerticalBox)
+			.RenderTransform_Static(&GetRenderTransform)
 			+SVerticalBox::Slot()
 			[
 				SNew(SAnimTest)
@@ -4099,6 +4489,7 @@ TSharedRef<SDockTab> SpawnTab(const FSpawnTabArgs& Args, FName TabIdentifier)
 		.Label( NSLOCTEXT("TestSuite1", "DocumentsTab", "Documents") )
 		[
 			SNew( SDocumentsTest, TabManagerRef )
+			.RenderTransform_Static(&GetRenderTransform)
 			.Tag("DocumentSpawner")
 		];
 	}
@@ -4127,6 +4518,7 @@ TSharedRef<SDockTab> SpawnTab(const FSpawnTabArgs& Args, FName TabIdentifier)
 			. Label( LOCTEXT("RichTextTestTab", "Rich Text") )
 			[
 				SNew( SRichTextTest )
+				.RenderTransform_Static(&GetRenderTransform)
 			];
 	}
 	else if ( TabIdentifier == FName( TEXT( "MultiLineEditTab" ) ) )
@@ -4139,6 +4531,7 @@ TSharedRef<SDockTab> SpawnTab(const FSpawnTabArgs& Args, FName TabIdentifier)
 				#else
 				SNew( SSpacer )
 				#endif //WITH_FANCY_TEXT
+				.RenderTransform_Static(&GetRenderTransform)
 			];
 	}
 #endif //WITH_FANCY_TEXT
@@ -4149,6 +4542,7 @@ TSharedRef<SDockTab> SpawnTab(const FSpawnTabArgs& Args, FName TabIdentifier)
 			. ToolTipText( LOCTEXT( "EditableTextTestTabToolTip", "Switches to the Editable Text tab, where you can test the various inline text editing controls." ) )
 		[
 			SNew( STextEditTest )
+			.RenderTransform_Static(&GetRenderTransform)
 		];
 	}
 #if WITH_FANCY_TEXT
@@ -4159,6 +4553,7 @@ TSharedRef<SDockTab> SpawnTab(const FSpawnTabArgs& Args, FName TabIdentifier)
 			. ToolTipText( LOCTEXT( "RichEditableTextTestTabToolTip", "Switches to the Rich Editable Text tab, where you can test the various rich editable text features." ) )
 			[
 				SNew( SRichTextEditTest )
+				.RenderTransform_Static(&GetRenderTransform)
 			];
 	}
 #endif //WITH_FANCY_TEXT
@@ -4168,6 +4563,7 @@ TSharedRef<SDockTab> SpawnTab(const FSpawnTabArgs& Args, FName TabIdentifier)
 			. Label( LOCTEXT("LayoutRoundingTab", "Layout Rounding") )
 		[
 			SNew( SLayoutRoundingTest )
+			.RenderTransform_Static(&GetRenderTransform)
 		];
 	}
 	else if (TabIdentifier == FName(TEXT("ElementTestsTab")))
@@ -4177,6 +4573,7 @@ TSharedRef<SDockTab> SpawnTab(const FSpawnTabArgs& Args, FName TabIdentifier)
 			. ToolTipText( LOCTEXT( "ElementTestsTabToolTip", "Switches to the Element Tests tab, which allows you to view various rendering-related features of Slate." ) )
 		[
 			SNew( SElementTesting )
+			.RenderTransform_Static(&GetRenderTransform)
 		];
 	}
 	else if (TabIdentifier == FName(TEXT("SplitterTestTab")))
@@ -4186,6 +4583,7 @@ TSharedRef<SDockTab> SpawnTab(const FSpawnTabArgs& Args, FName TabIdentifier)
 			. ToolTipText( LOCTEXT( "SplitterTestTabToolTip", "Switches to the Splitter Test tab, which you can use to test splitters." ) )
 		[
 			SNew( SSplitterTest )
+			.RenderTransform_Static(&GetRenderTransform)
 		];
 	}
 	else if (TabIdentifier == FName(TEXT("MultiBoxTestTab")))
@@ -4195,6 +4593,7 @@ TSharedRef<SDockTab> SpawnTab(const FSpawnTabArgs& Args, FName TabIdentifier)
 			. ToolTipText( LOCTEXT( "MultiBoxTextTabToolTip", "Switches to the MultiBox tab, where you can test out MultiBoxes and MultiBlocks." ) )
 		[
 			SNew( SMultiBoxTest )
+			.RenderTransform_Static(&GetRenderTransform)
 		];
 	}
 	else if (TabIdentifier == FName(TEXT("WidgetGalleryTab")))
@@ -4213,6 +4612,7 @@ TSharedRef<SDockTab> SpawnTab(const FSpawnTabArgs& Args, FName TabIdentifier)
 			. ToolTipText( LOCTEXT( "ColorPickerTestTabToolTip", "Switches to the Color Picker tab, where you can test out the color picker." ) )
 		[
 			SNew(SColorPickerTest)
+			.RenderTransform_Static(&GetRenderTransform)
 		];
 	}
 	else if (TabIdentifier == "DPIScalingTest")
@@ -4220,6 +4620,7 @@ TSharedRef<SDockTab> SpawnTab(const FSpawnTabArgs& Args, FName TabIdentifier)
 		return SNew(SDockTab)
 		[
 			SNew(SDPIScalingTest)
+			.RenderTransform_Static(&GetRenderTransform)
 		];
 	}
 	else if (TabIdentifier == FName(TEXT("NotificationListTestTab")))
@@ -4229,6 +4630,7 @@ TSharedRef<SDockTab> SpawnTab(const FSpawnTabArgs& Args, FName TabIdentifier)
 			. ToolTipText( LOCTEXT( "NotificationListTestTabToolTip", "Switches to the Notification List tab, where you can test out the notification list." ) )
 			[
 				SNew(SNotificationListTest)
+				.RenderTransform_Static(&GetRenderTransform)
 			];
 	}
 	else if (TabIdentifier == FName("GridPanelTest"))
@@ -4236,6 +4638,7 @@ TSharedRef<SDockTab> SpawnTab(const FSpawnTabArgs& Args, FName TabIdentifier)
 		return SNew(SDockTab)
 			[
 				SNew(SGridPanelTest)
+				.RenderTransform_Static(&GetRenderTransform)
 			];
 	}
 	else
@@ -4479,6 +4882,22 @@ TSharedRef<SDockTab> SpawnTestSuite2( const FSpawnTabArgs& Args )
 	return TestSuite2Tab;
 }
 
+TSharedRef<SDockTab> SpawnRenderTransformManipulator( const FSpawnTabArgs& Args )
+{	
+	TSharedRef<SDockTab> RenderTransformManipulatorTab =
+		SNew(SDockTab)
+		.TabRole(ETabRole::MajorTab)
+		.Label(LOCTEXT("RenderTransformTabLabel", "Render Transform"))
+		.ToolTipText(LOCTEXT("RenderTransformTabToolTip", "Allows manipulating the render transform of all test tabs at a global level."));
+
+	RenderTransformManipulatorTab->SetContent
+	(
+		SNew( SRenderTransformManipulatorWidget )
+	);
+
+	return RenderTransformManipulatorTab;
+}
+
 
 void RestoreSlateTestSuite()
 {
@@ -4486,6 +4905,7 @@ void RestoreSlateTestSuite()
 
 	FGlobalTabmanager::Get()->RegisterTabSpawner("TestSuite1", FOnSpawnTab::CreateStatic( &SpawnTestSuite1 ) );
 	FGlobalTabmanager::Get()->RegisterTabSpawner("TestSuite2", FOnSpawnTab::CreateStatic( &SpawnTestSuite2 ) );
+	FGlobalTabmanager::Get()->RegisterTabSpawner("RenderTrasnformManipulator", FOnSpawnTab::CreateStatic(&SpawnRenderTransformManipulator));
 
 	TSharedRef<FTabManager::FLayout> Layout = FTabManager::NewLayout( "SlateTestSuite_Layout" )
 	->AddArea
@@ -4497,6 +4917,7 @@ void RestoreSlateTestSuite()
 			FTabManager::NewStack()
 			->AddTab( "TestSuite2", ETabState::OpenedTab )
 			->AddTab( "TestSuite1", ETabState::OpenedTab )
+			->AddTab( "RenderTrasnformManipulator", ETabState::OpenedTab )
 		)		
 	)
 	#if PLATFORM_SUPPORTS_MULTIPLE_NATIVE_WINDOWS
diff --git a/Engine/Source/Runtime/Core/Private/Math/UnrealMath.cpp b/Engine/Source/Runtime/Core/Private/Math/UnrealMath.cpp
--- a/Engine/Source/Runtime/Core/Private/Math/UnrealMath.cpp
+++ b/Engine/Source/Runtime/Core/Private/Math/UnrealMath.cpp
@@ -380,9 +380,18 @@ void FMatrix::DebugPrint() const
 
 FQuat FQuat::MakeFromEuler(const FVector& Euler)
 {
-	return FQuat( FRotationTranslationMatrix( FRotator::MakeFromEuler(Euler), FVector::ZeroVector ) );
+	return FRotator::MakeFromEuler(Euler).Quaternion();
 }
 
+FMatrix FRotationAboutPointMatrix::Make(const FQuat& Rot, const FVector& Origin)
+{
+	return FRotationAboutPointMatrix(Rot.Rotator(), Origin);
+}
+
+FMatrix FRotationMatrix::Make(FQuat const& Rot)
+{
+	return FQuatRotationTranslationMatrix(Rot, FVector::ZeroVector);
+}
 
 FMatrix FRotationMatrix::MakeFromX(FVector const& XAxis)
 {
diff --git a/Engine/Source/Runtime/Core/Public/Math/QuatRotationTranslationMatrix.h b/Engine/Source/Runtime/Core/Public/Math/QuatRotationTranslationMatrix.h
--- a/Engine/Source/Runtime/Core/Public/Math/QuatRotationTranslationMatrix.h
+++ b/Engine/Source/Runtime/Core/Public/Math/QuatRotationTranslationMatrix.h
@@ -12,6 +12,13 @@ public:
 	* @param Origin translation to apply
 	*/
 	FQuatRotationTranslationMatrix(const FQuat& Q, const FVector& Origin);
+
+	/** Matrix factory. Return an FMatrix so we don't have type conversion issues in expressions. */
+	static FMatrix Make(const FQuat& Q, const FVector& Origin)
+	{
+		return FQuatRotationTranslationMatrix(Q, Origin);
+	}
+
 };
 
 FORCEINLINE FQuatRotationTranslationMatrix::FQuatRotationTranslationMatrix(const FQuat& Q, const FVector& Origin)
diff --git a/Engine/Source/Runtime/Core/Public/Math/RotationAboutPointMatrix.h b/Engine/Source/Runtime/Core/Public/Math/RotationAboutPointMatrix.h
new file mode 100644
--- /dev/null
+++ b/Engine/Source/Runtime/Core/Public/Math/RotationAboutPointMatrix.h
@@ -0,0 +1,38 @@
+// Copyright 1998-2014 Epic Games, Inc. All Rights Reserved.
+
+#pragma once
+
+/** Rotates about an Origin point. */
+class FRotationAboutPointMatrix : public FRotationTranslationMatrix
+{
+public:
+	/** Constructor
+	 *
+	 * @param Rot rotation
+	 * @param Origin about which to rotate.
+	 */
+	FRotationAboutPointMatrix(const FRotator& Rot, const FVector& Origin);
+
+	/** Matrix factory. Return an FMatrix so we don't have type conversion issues in expressions. */
+	static FMatrix Make(const FRotator& Rot, const FVector& Origin)
+	{
+		return FRotationAboutPointMatrix(Rot, Origin);
+	}
+
+	/** Matrix factory. Return an FMatrix so we don't have type conversion issues in expressions. */
+	static CORE_API FMatrix Make(const FQuat& Rot, const FVector& Origin);
+};
+
+FORCEINLINE FRotationAboutPointMatrix::FRotationAboutPointMatrix(const FRotator& Rot, const FVector& Origin)
+	:FRotationTranslationMatrix(Rot, Origin)
+{
+	// FRotationTranslationMatrix generates R * T.
+	// We need -T * R * T, so prepend that translation:
+	FVector XAxis(M[0][0], M[1][0], M[2][0]);
+	FVector YAxis(M[0][1], M[1][1], M[2][1]);
+	FVector ZAxis(M[0][2], M[1][2], M[2][2]);
+
+	M[3][0]	-= XAxis | Origin;
+	M[3][1]	-= YAxis | Origin;
+	M[3][2]	-= ZAxis | Origin;
+}
\ No newline at end of file
diff --git a/Engine/Source/Runtime/Core/Public/Math/RotationMatrix.h b/Engine/Source/Runtime/Core/Public/Math/RotationMatrix.h
--- a/Engine/Source/Runtime/Core/Public/Math/RotationMatrix.h
+++ b/Engine/Source/Runtime/Core/Public/Math/RotationMatrix.h
@@ -12,6 +12,15 @@ public:
 	 */
 	FRotationMatrix(const FRotator& Rot) : FRotationTranslationMatrix(Rot, FVector::ZeroVector)	{}
 
+	/** Matrix factory. Return an FMatrix so we don't have type conversion issues in expressions. */
+	static FMatrix Make(FRotator const& Rot)
+	{
+		return FRotationMatrix(Rot);
+	}
+
+	/** Matrix factory. Return an FMatrix so we don't have type conversion issues in expressions. */
+	static CORE_API FMatrix Make(FQuat const& Rot);
+
 	/** Builds a rotation matrix given only a XAxis. Y and Z are unspecified but will be orthonormal. XAxis need not be normalized. */
 	static CORE_API FMatrix MakeFromX(FVector const& XAxis);
 
diff --git a/Engine/Source/Runtime/Core/Public/Math/RotationTranslationMatrix.h b/Engine/Source/Runtime/Core/Public/Math/RotationTranslationMatrix.h
--- a/Engine/Source/Runtime/Core/Public/Math/RotationTranslationMatrix.h
+++ b/Engine/Source/Runtime/Core/Public/Math/RotationTranslationMatrix.h
@@ -12,6 +12,13 @@ public:
 	 * @param Origin translation to apply
 	 */
 	FRotationTranslationMatrix(const FRotator& Rot, const FVector& Origin);
+
+	/** Matrix factory. Return an FMatrix so we don't have type conversion issues in expressions. */
+	static FMatrix Make(const FRotator& Rot, const FVector& Origin)
+	{
+		return FRotationTranslationMatrix(Rot, Origin);
+	}
+
 };
 
 FORCEINLINE FRotationTranslationMatrix::FRotationTranslationMatrix(const FRotator& Rot, const FVector& Origin)
diff --git a/Engine/Source/Runtime/Core/Public/Math/ScaleMatrix.h b/Engine/Source/Runtime/Core/Public/Math/ScaleMatrix.h
--- a/Engine/Source/Runtime/Core/Public/Math/ScaleMatrix.h
+++ b/Engine/Source/Runtime/Core/Public/Math/ScaleMatrix.h
@@ -24,6 +24,18 @@ public:
 	 * @param Scale Non-uniform scale to apply to matrix.
 	 */
 	FScaleMatrix( const FVector& Scale );
+
+	/** Matrix factory. Return an FMatrix so we don't have type conversion issues in expressions. */
+	static FMatrix Make(float Scale)
+	{
+		return FScaleMatrix(Scale);
+	}
+
+	/** Matrix factory. Return an FMatrix so we don't have type conversion issues in expressions. */
+	static FMatrix Make(const FVector& Scale)
+	{
+		return FScaleMatrix(Scale);
+	}
 };
 
 
