Adding RenderTransformPivot to SWidget, which allows you to specify a pivot in normalized local space. Fairly trivial change to concatenation logic in FGeometry, but does imply some additional math every time we arrange children. Most of the changes are due to places where people call SWidget::Construct, which makes this look more painful to add than it really was.

diff --git a/Engine/Source/Runtime/AppFramework/Private/Framework/Testing/SLayoutExample.cpp b/Engine/Source/Runtime/AppFramework/Private/Framework/Testing/SLayoutExample.cpp
--- a/Engine/Source/Runtime/AppFramework/Private/Framework/Testing/SLayoutExample.cpp
+++ b/Engine/Source/Runtime/AppFramework/Private/Framework/Testing/SLayoutExample.cpp
@@ -397,10 +397,12 @@ public:
 
 APPFRAMEWORK_API TSharedRef<SWidget> MakeLayoutExample()
 {
-	extern FSlateRenderTransform GetRenderTransform();
-	return 
+	extern FSlateRenderTransform GetTestRenderTransform();
+	extern FVector2D GetTestRenderTransformPivot();
+	return
 		SNew(SExampleLayout)
-		.RenderTransform_Static(&GetRenderTransform);
+		.RenderTransform_Static(&GetTestRenderTransform)
+		.RenderTransformPivot_Static(&GetTestRenderTransformPivot);
 }
 
 #undef LOCTEXT_NAMESPACE
\ No newline at end of file
diff --git a/Engine/Source/Runtime/AppFramework/Private/Framework/Testing/STableViewTesting.cpp b/Engine/Source/Runtime/AppFramework/Private/Framework/Testing/STableViewTesting.cpp
--- a/Engine/Source/Runtime/AppFramework/Private/Framework/Testing/STableViewTesting.cpp
+++ b/Engine/Source/Runtime/AppFramework/Private/Framework/Testing/STableViewTesting.cpp
@@ -1102,10 +1102,12 @@ protected:
 
 TSharedRef<SWidget> MakeTableViewTesting()
 {
-	extern FSlateRenderTransform GetRenderTransform();
-	return 
-		SNew( STableViewTesting )
-		.RenderTransform_Static(&GetRenderTransform);
+	extern FSlateRenderTransform GetTestRenderTransform();
+	extern FVector2D GetTestRenderTransformPivot();
+	return
+		SNew(STableViewTesting)
+		.RenderTransform_Static(&GetTestRenderTransform)
+		.RenderTransformPivot_Static(&GetTestRenderTransformPivot);
 }
 
 
diff --git a/Engine/Source/Runtime/AppFramework/Private/Framework/Testing/SWidgetGallery.cpp b/Engine/Source/Runtime/AppFramework/Private/Framework/Testing/SWidgetGallery.cpp
--- a/Engine/Source/Runtime/AppFramework/Private/Framework/Testing/SWidgetGallery.cpp
+++ b/Engine/Source/Runtime/AppFramework/Private/Framework/Testing/SWidgetGallery.cpp
@@ -1100,10 +1100,12 @@ private:
  */
 TSharedRef<SWidget> MakeWidgetGallery()
 {
-	extern FSlateRenderTransform GetRenderTransform();
-	return 
+	extern FSlateRenderTransform GetTestRenderTransform();
+	extern FVector2D GetTestRenderTransformPivot();
+	return
 		SNew(SWidgetGallery)
-		.RenderTransform_Static(&GetRenderTransform);
+		.RenderTransform_Static(&GetTestRenderTransform)
+		.RenderTransformPivot_Static(&GetTestRenderTransformPivot);
 }
 
 
diff --git a/Engine/Source/Runtime/AppFramework/Private/Widgets/Testing/STestSuite.cpp b/Engine/Source/Runtime/AppFramework/Private/Widgets/Testing/STestSuite.cpp
--- a/Engine/Source/Runtime/AppFramework/Private/Widgets/Testing/STestSuite.cpp
+++ b/Engine/Source/Runtime/AppFramework/Private/Widgets/Testing/STestSuite.cpp
@@ -4331,9 +4331,14 @@ namespace
 /**
  * Global access here because we need other translation units to access this function.
  */
-FSlateRenderTransform GetRenderTransform()
+FSlateRenderTransform GetTestRenderTransform()
 {
-	return Concatenate(-Offset, Concatenate(Shear, Scale, Rot, Offset));
+	return TransformCast<FSlateRenderTransform>(Concatenate(Shear, Scale, Rot));
+}
+
+FVector2D GetTestRenderTransformPivot()
+{
+	return Offset;
 }
 
 void SRenderTransformManipulatorWidgetImpl::Construct(const FArguments& InArgs)
@@ -4344,7 +4349,8 @@ void SRenderTransformManipulatorWidgetImpl::Construct(const FArguments& InArgs)
 			+ SVerticalBox::Slot().Padding(4)
 				[
 					SAssignNew(ImageWidget, SImage)
-					.RenderTransform_Static(&::GetRenderTransform)
+					.RenderTransform_Static(&::GetTestRenderTransform)
+					.RenderTransformPivot_Static(&::GetTestRenderTransformPivot)
 					.Image(GetStyle().GetBrush("UE4Icon"))
 				]
 			+ SVerticalBox::Slot().AutoHeight().Padding(4)
@@ -4373,8 +4379,8 @@ void SRenderTransformManipulatorWidgetImpl::Construct(const FArguments& InArgs)
 					+ SHorizontalBox::Slot().HAlign(HAlign_Fill)
 						[
 							SNew(SSpinBox<float>)
-							.MinValue(-1000.0f)
-							.MaxValue(1000.0f)
+							.MinValue(-2.0f)
+							.MaxValue(2.0f)
 							.OnValueChanged_Static([](float val) { Offset.X = val; })
 							.Value_Static([] { return Offset.X; })
 						]
@@ -4389,8 +4395,8 @@ void SRenderTransformManipulatorWidgetImpl::Construct(const FArguments& InArgs)
 					+ SHorizontalBox::Slot().HAlign(HAlign_Fill)
 						[
 							SNew(SSpinBox<float>)
-							.MinValue(-1000.0f)
-							.MaxValue(1000.0f)
+							.MinValue(-2.0f)
+							.MaxValue(2.0f)
 							.OnValueChanged_Static([](float val) { Offset.Y = val; })
 							.Value_Static([] { return Offset.Y; })
 						]
@@ -4476,7 +4482,8 @@ static TSharedPtr<FTabManager> TestSuite1TabManager;
 BEGIN_SLATE_FUNCTION_BUILD_OPTIMIZATION
 TSharedRef<SDockTab> SpawnTab(const FSpawnTabArgs& Args, FName TabIdentifier)
 {
-	extern FSlateRenderTransform GetRenderTransform();
+	extern FSlateRenderTransform GetTestRenderTransform();
+	extern FVector2D GetTestRenderTransformPivot();
 	if (TabIdentifier == FName(TEXT("AnimationTestTab")))
 	{
 		return SNew(SDockTab)
@@ -4505,7 +4512,8 @@ TSharedRef<SDockTab> SpawnTab(const FSpawnTabArgs& Args, FName TabIdentifier)
 			)
 		[
 			SNew(SVerticalBox)
-			.RenderTransform_Static(&GetRenderTransform)
+			.RenderTransform_Static(&::GetTestRenderTransform)
+			.RenderTransformPivot_Static(&::GetTestRenderTransformPivot)
 			+SVerticalBox::Slot()
 			[
 				SNew(SAnimTest)
@@ -4525,7 +4533,8 @@ TSharedRef<SDockTab> SpawnTab(const FSpawnTabArgs& Args, FName TabIdentifier)
 		.Label( NSLOCTEXT("TestSuite1", "DocumentsTab", "Documents") )
 		[
 			SNew( SDocumentsTest, TabManagerRef )
-			.RenderTransform_Static(&GetRenderTransform)
+			.RenderTransform_Static(&::GetTestRenderTransform)
+			.RenderTransformPivot_Static(&::GetTestRenderTransformPivot)
 			.Tag("DocumentSpawner")
 		];
 	}
@@ -4554,7 +4563,8 @@ TSharedRef<SDockTab> SpawnTab(const FSpawnTabArgs& Args, FName TabIdentifier)
 			. Label( LOCTEXT("RichTextTestTab", "Rich Text") )
 			[
 				SNew( SRichTextTest )
-				.RenderTransform_Static(&GetRenderTransform)
+				.RenderTransform_Static(&::GetTestRenderTransform)
+				.RenderTransformPivot_Static(&::GetTestRenderTransformPivot)
 			];
 	}
 	else if ( TabIdentifier == FName( TEXT( "MultiLineEditTab" ) ) )
@@ -4567,7 +4577,8 @@ TSharedRef<SDockTab> SpawnTab(const FSpawnTabArgs& Args, FName TabIdentifier)
 				#else
 				SNew( SSpacer )
 				#endif //WITH_FANCY_TEXT
-				.RenderTransform_Static(&GetRenderTransform)
+				.RenderTransform_Static(&::GetTestRenderTransform)
+				.RenderTransformPivot_Static(&::GetTestRenderTransformPivot)
 			];
 	}
 #endif //WITH_FANCY_TEXT
@@ -4578,7 +4589,8 @@ TSharedRef<SDockTab> SpawnTab(const FSpawnTabArgs& Args, FName TabIdentifier)
 			. ToolTipText( LOCTEXT( "EditableTextTestTabToolTip", "Switches to the Editable Text tab, where you can test the various inline text editing controls." ) )
 		[
 			SNew( STextEditTest )
-			.RenderTransform_Static(&GetRenderTransform)
+			.RenderTransform_Static(&::GetTestRenderTransform)
+			.RenderTransformPivot_Static(&::GetTestRenderTransformPivot)
 		];
 	}
 #if WITH_FANCY_TEXT
@@ -4589,7 +4601,8 @@ TSharedRef<SDockTab> SpawnTab(const FSpawnTabArgs& Args, FName TabIdentifier)
 			. ToolTipText( LOCTEXT( "RichEditableTextTestTabToolTip", "Switches to the Rich Editable Text tab, where you can test the various rich editable text features." ) )
 			[
 				SNew( SRichTextEditTest )
-				.RenderTransform_Static(&GetRenderTransform)
+				.RenderTransform_Static(&::GetTestRenderTransform)
+				.RenderTransformPivot_Static(&::GetTestRenderTransformPivot)
 			];
 	}
 #endif //WITH_FANCY_TEXT
@@ -4599,7 +4612,8 @@ TSharedRef<SDockTab> SpawnTab(const FSpawnTabArgs& Args, FName TabIdentifier)
 			. Label( LOCTEXT("LayoutRoundingTab", "Layout Rounding") )
 		[
 			SNew( SLayoutRoundingTest )
-			.RenderTransform_Static(&GetRenderTransform)
+			.RenderTransform_Static(&::GetTestRenderTransform)
+			.RenderTransformPivot_Static(&::GetTestRenderTransformPivot)
 		];
 	}
 	else if (TabIdentifier == FName(TEXT("ElementTestsTab")))
@@ -4609,7 +4623,8 @@ TSharedRef<SDockTab> SpawnTab(const FSpawnTabArgs& Args, FName TabIdentifier)
 			. ToolTipText( LOCTEXT( "ElementTestsTabToolTip", "Switches to the Element Tests tab, which allows you to view various rendering-related features of Slate." ) )
 		[
 			SNew( SElementTesting )
-			.RenderTransform_Static(&GetRenderTransform)
+			.RenderTransform_Static(&::GetTestRenderTransform)
+			.RenderTransformPivot_Static(&::GetTestRenderTransformPivot)
 		];
 	}
 	else if (TabIdentifier == FName(TEXT("SplitterTestTab")))
@@ -4619,7 +4634,8 @@ TSharedRef<SDockTab> SpawnTab(const FSpawnTabArgs& Args, FName TabIdentifier)
 			. ToolTipText( LOCTEXT( "SplitterTestTabToolTip", "Switches to the Splitter Test tab, which you can use to test splitters." ) )
 		[
 			SNew( SSplitterTest )
-			.RenderTransform_Static(&GetRenderTransform)
+			.RenderTransform_Static(&::GetTestRenderTransform)
+			.RenderTransformPivot_Static(&::GetTestRenderTransformPivot)
 		];
 	}
 	else if (TabIdentifier == FName(TEXT("MultiBoxTestTab")))
@@ -4629,7 +4645,8 @@ TSharedRef<SDockTab> SpawnTab(const FSpawnTabArgs& Args, FName TabIdentifier)
 			. ToolTipText( LOCTEXT( "MultiBoxTextTabToolTip", "Switches to the MultiBox tab, where you can test out MultiBoxes and MultiBlocks." ) )
 		[
 			SNew( SMultiBoxTest )
-			.RenderTransform_Static(&GetRenderTransform)
+			.RenderTransform_Static(&::GetTestRenderTransform)
+			.RenderTransformPivot_Static(&::GetTestRenderTransformPivot)
 		];
 	}
 	else if (TabIdentifier == FName(TEXT("WidgetGalleryTab")))
@@ -4648,7 +4665,8 @@ TSharedRef<SDockTab> SpawnTab(const FSpawnTabArgs& Args, FName TabIdentifier)
 			. ToolTipText( LOCTEXT( "ColorPickerTestTabToolTip", "Switches to the Color Picker tab, where you can test out the color picker." ) )
 		[
 			SNew(SColorPickerTest)
-			.RenderTransform_Static(&GetRenderTransform)
+			.RenderTransform_Static(&::GetTestRenderTransform)
+			.RenderTransformPivot_Static(&::GetTestRenderTransformPivot)
 		];
 	}
 	else if (TabIdentifier == "DPIScalingTest")
@@ -4656,7 +4674,8 @@ TSharedRef<SDockTab> SpawnTab(const FSpawnTabArgs& Args, FName TabIdentifier)
 		return SNew(SDockTab)
 		[
 			SNew(SDPIScalingTest)
-			.RenderTransform_Static(&GetRenderTransform)
+			.RenderTransform_Static(&::GetTestRenderTransform)
+			.RenderTransformPivot_Static(&::GetTestRenderTransformPivot)
 		];
 	}
 	else if (TabIdentifier == FName(TEXT("NotificationListTestTab")))
@@ -4666,7 +4685,8 @@ TSharedRef<SDockTab> SpawnTab(const FSpawnTabArgs& Args, FName TabIdentifier)
 			. ToolTipText( LOCTEXT( "NotificationListTestTabToolTip", "Switches to the Notification List tab, where you can test out the notification list." ) )
 			[
 				SNew(SNotificationListTest)
-				.RenderTransform_Static(&GetRenderTransform)
+				.RenderTransform_Static(&::GetTestRenderTransform)
+				.RenderTransformPivot_Static(&::GetTestRenderTransformPivot)
 			];
 	}
 	else if (TabIdentifier == FName("GridPanelTest"))
@@ -4674,7 +4694,8 @@ TSharedRef<SDockTab> SpawnTab(const FSpawnTabArgs& Args, FName TabIdentifier)
 		return SNew(SDockTab)
 			[
 				SNew(SGridPanelTest)
-				.RenderTransform_Static(&GetRenderTransform)
+				.RenderTransform_Static(&::GetTestRenderTransform)
+				.RenderTransformPivot_Static(&::GetTestRenderTransformPivot)
 			];
 	}
 	else
diff --git a/Engine/Source/Runtime/Slate/Private/Widgets/Input/SSlider.cpp b/Engine/Source/Runtime/Slate/Private/Widgets/Input/SSlider.cpp
--- a/Engine/Source/Runtime/Slate/Private/Widgets/Input/SSlider.cpp
+++ b/Engine/Source/Runtime/Slate/Private/Widgets/Input/SSlider.cpp
@@ -64,7 +64,7 @@ int32 SSlider::OnPaint( const FPaintArgs& Args, const FGeometry& AllottedGeometr
 		SliderGeometry = AllottedGeometry.MakeChild(
 			FVector2D(AllottedWidth, AllottedHeight), 
 			FSlateLayoutTransform(), 
-			RenderTransform);
+			RenderTransform, FVector2D::ZeroVector);
 		// The clipping rect is already given properly in window space. But we do not support layout rotations, so our local space rendering cannot
 		// get the clipping rect into local space properly for the local space clipping we do in the shader.
 		// Thus, we transform the clip coords into local space manually, UNDO the render transform so it will clip properly,
diff --git a/Engine/Source/Runtime/Slate/Private/Widgets/Layout/SSpacer.cpp b/Engine/Source/Runtime/Slate/Private/Widgets/Layout/SSpacer.cpp
--- a/Engine/Source/Runtime/Slate/Private/Widgets/Layout/SSpacer.cpp
+++ b/Engine/Source/Runtime/Slate/Private/Widgets/Layout/SSpacer.cpp
@@ -19,6 +19,7 @@ void SSpacer::Construct( const FArguments& InArgs )
 		InArgs._IsEnabled,
 		InArgs._Visibility,
 		InArgs._RenderTransform,
+		InArgs._RenderTransformPivot,
 		InArgs._Tag,
 		InArgs.MetaData);
 }
diff --git a/Engine/Source/Runtime/Slate/Private/Widgets/Views/SHeaderRow.cpp b/Engine/Source/Runtime/Slate/Private/Widgets/Views/SHeaderRow.cpp
--- a/Engine/Source/Runtime/Slate/Private/Widgets/Views/SHeaderRow.cpp
+++ b/Engine/Source/Runtime/Slate/Private/Widgets/Views/SHeaderRow.cpp
@@ -34,7 +34,7 @@ public:
 	{
 		check(InArgs._Style);
 
-		SWidget::Construct( InArgs._ToolTipText, InArgs._ToolTip, InArgs._Cursor, InArgs._IsEnabled, InArgs._Visibility, InArgs._RenderTransform, InArgs._Tag, InArgs.MetaData );
+		SWidget::Construct( InArgs._ToolTipText, InArgs._ToolTip, InArgs._Cursor, InArgs._IsEnabled, InArgs._Visibility, InArgs._RenderTransform, InArgs._RenderTransformPivot, InArgs._Tag, InArgs.MetaData );
 
 		Style = InArgs._Style;
 		ColumnId = Column.ColumnId;
@@ -384,7 +384,7 @@ void SHeaderRow::Construct( const FArguments& InArgs )
 {
 	check(InArgs._Style);
 
-	SWidget::Construct( InArgs._ToolTipText, InArgs._ToolTip, InArgs._Cursor, InArgs._IsEnabled, InArgs._Visibility, InArgs._RenderTransform, InArgs._Tag, InArgs.MetaData );
+	SWidget::Construct( InArgs._ToolTipText, InArgs._ToolTip, InArgs._Cursor, InArgs._IsEnabled, InArgs._Visibility, InArgs._RenderTransform, InArgs._RenderTransformPivot, InArgs._Tag, InArgs.MetaData );
 
 	ScrollBarThickness = FVector2D::ZeroVector;
 	ScrollBarVisibility = EVisibility::Collapsed;
diff --git a/Engine/Source/Runtime/SlateCore/Private/Layout/Geometry.cpp b/Engine/Source/Runtime/SlateCore/Private/Layout/Geometry.cpp
--- a/Engine/Source/Runtime/SlateCore/Private/Layout/Geometry.cpp
+++ b/Engine/Source/Runtime/SlateCore/Private/Layout/Geometry.cpp
@@ -39,11 +39,22 @@ FGeometry::FGeometry(const FVector2D& OffsetFromParent, const FVector2D& ParentA
 	const_cast<FVector2D&>(Position) = LocalLayoutTransform.GetTranslation();
 }
 
-FGeometry::FGeometry(const FVector2D& InLocalSize, const FSlateLayoutTransform& InLocalLayoutTransform, const FSlateRenderTransform& InLocalRenderTransform, const FSlateLayoutTransform& ParentAccumulatedLayoutTransform, const FSlateRenderTransform& ParentAccumulatedRenderTransform) 
+FGeometry::FGeometry(const FVector2D& InLocalSize, const FSlateLayoutTransform& InLocalLayoutTransform, const FSlateRenderTransform& InLocalRenderTransform, const FVector2D& InLocalRenderTransformPivot, const FSlateLayoutTransform& ParentAccumulatedLayoutTransform, const FSlateRenderTransform& ParentAccumulatedRenderTransform) 
 	: Size(InLocalSize)
 	, Scale(1.0f)
 	, AbsolutePosition(0.0f, 0.0f)
-	, AccumulatedRenderTransform(Concatenate(InLocalRenderTransform, InLocalLayoutTransform, ParentAccumulatedRenderTransform))
+	, AccumulatedRenderTransform(
+		Concatenate(
+			// convert the pivot to local space and make it the origin
+			Inverse(TransformPoint(FScale2D(InLocalSize), InLocalRenderTransformPivot)), 
+			// apply the render transform in local space centered around the pivot
+			InLocalRenderTransform, 
+			// translate the pivot point back.
+			TransformPoint(FScale2D(InLocalSize), InLocalRenderTransformPivot), 
+			// apply the layout transform next.
+			InLocalLayoutTransform, 
+			// finally apply the parent accumulated transform, which takes us to the root.
+			ParentAccumulatedRenderTransform))
 {
 	FSlateLayoutTransform AccumulatedLayoutTransform = Concatenate(InLocalLayoutTransform, ParentAccumulatedLayoutTransform);
 	// HACK to allow us to make FGeometry public members immutable to catch misuse.
@@ -78,9 +89,9 @@ FGeometry FGeometry::MakeRoot(const FVector2D& LocalSize, const FSlateLayoutTran
 	return FGeometry(LocalSize, LayoutTransform, FSlateLayoutTransform(), FSlateRenderTransform());
 }
 
-FGeometry FGeometry::MakeChild(const FVector2D& LocalSize, const FSlateLayoutTransform& LayoutTransform, const FSlateRenderTransform& RenderTransform) const
+FGeometry FGeometry::MakeChild(const FVector2D& LocalSize, const FSlateLayoutTransform& LayoutTransform, const FSlateRenderTransform& RenderTransform, const FVector2D& RenderTransformPivot) const
 {
-	return FGeometry(LocalSize, LayoutTransform, RenderTransform, GetAccumulatedLayoutTransform(), GetAccumulatedRenderTransform());
+	return FGeometry(LocalSize, LayoutTransform, RenderTransform, RenderTransformPivot, GetAccumulatedLayoutTransform(), GetAccumulatedRenderTransform());
 }
 
 FGeometry FGeometry::MakeChild( const FVector2D& LocalSize, const FSlateLayoutTransform& LayoutTransform ) const
@@ -95,7 +106,7 @@ FArrangedWidget FGeometry::MakeChild(const TSharedRef<SWidget>& ChildWidget, con
 
 FArrangedWidget FGeometry::MakeChild(const TSharedRef<SWidget>& ChildWidget, const FVector2D& LocalSize, const FSlateLayoutTransform& LayoutTransform) const
 {
-	return FArrangedWidget(ChildWidget, MakeChild(LocalSize, LayoutTransform, ChildWidget->GetRenderTransform()));
+	return FArrangedWidget(ChildWidget, MakeChild(LocalSize, LayoutTransform, ChildWidget->GetRenderTransform(), ChildWidget->GetRenderTransformPivot()));
 }
 
 FGeometry FGeometry::MakeChild( const FVector2D& ChildOffset, const FVector2D& LocalSize, float LocalScale ) const
@@ -109,7 +120,7 @@ FArrangedWidget FGeometry::MakeChild( const TSharedRef<SWidget>& ChildWidget, co
 {
 	// Since ChildOffset is given as a LocalSpaceOffset, we MUST convert this offset into the space of the parent to construct a valid layout transform.
 	// The extra TransformPoint below does this by converting the local offset to an offset in parent space.
-	return FArrangedWidget(ChildWidget, MakeChild(LocalSize, FSlateLayoutTransform(ChildScale, TransformPoint(ChildScale, ChildOffset)), ChildWidget->GetRenderTransform()));
+	return FArrangedWidget(ChildWidget, MakeChild(LocalSize, FSlateLayoutTransform(ChildScale, TransformPoint(ChildScale, ChildOffset)), ChildWidget->GetRenderTransform(), ChildWidget->GetRenderTransformPivot()));
 }
 
 FPaintGeometry FGeometry::ToPaintGeometry() const
diff --git a/Engine/Source/Runtime/SlateCore/Private/Widgets/SNullWidget.cpp b/Engine/Source/Runtime/SlateCore/Private/Widgets/SNullWidget.cpp
--- a/Engine/Source/Runtime/SlateCore/Private/Widgets/SNullWidget.cpp
+++ b/Engine/Source/Runtime/SlateCore/Private/Widgets/SNullWidget.cpp
@@ -25,6 +25,7 @@ public:
 			InArgs._IsEnabled,
 			InArgs._Visibility,
 			InArgs._RenderTransform,
+			InArgs._RenderTransformPivot,
 			InArgs._Tag,
 			InArgs.MetaData
 		);
diff --git a/Engine/Source/Runtime/SlateCore/Private/Widgets/SWidget.cpp b/Engine/Source/Runtime/SlateCore/Private/Widgets/SWidget.cpp
--- a/Engine/Source/Runtime/SlateCore/Private/Widgets/SWidget.cpp
+++ b/Engine/Source/Runtime/SlateCore/Private/Widgets/SWidget.cpp
@@ -17,7 +17,8 @@ SWidget::SWidget()
 	, Cursor( TOptional<EMouseCursor::Type>() )
 	, EnabledState( true )
 	, Visibility( EVisibility::Visible )
-	, RenderTransform()
+	, RenderTransform( FSlateRenderTransform() )
+	, RenderTransformPivot( FVector2D::ZeroVector )
 	, DesiredSize(FVector2D::ZeroVector)
 	, ToolTip()
 	, bToolTipForceFieldEnabled( false )
@@ -33,6 +34,7 @@ void SWidget::Construct(
 	const TAttribute<bool> & InEnabledState ,
 	const TAttribute<EVisibility> & InVisibility,
 	const TAttribute<FSlateRenderTransform>& InTransform,
+	const TAttribute<FVector2D>& InTransformPivot,
 	const FName& InTag,
 	const TArray<TSharedRef<ISlateMetaData>>& InMetaData
 )
@@ -57,6 +59,7 @@ void SWidget::Construct(
 	EnabledState = InEnabledState;
 	Visibility = InVisibility;
 	RenderTransform = InTransform;
+	RenderTransformPivot = InTransformPivot;
 	Tag = InTag;
 	MetaData = InMetaData;
 }
diff --git a/Engine/Source/Runtime/SlateCore/Public/Layout/Geometry.h b/Engine/Source/Runtime/SlateCore/Public/Layout/Geometry.h
--- a/Engine/Source/Runtime/SlateCore/Public/Layout/Geometry.h
+++ b/Engine/Source/Runtime/SlateCore/Public/Layout/Geometry.h
@@ -56,6 +56,7 @@ private:
 	 * @param InLocalSize						The size of the geometry in Local Space.
 	 * @param InLocalLayoutTransform			A layout transform from local space to the parent geoemtry's local space.
 	 * @param InLocalRenderTransform			A render-only transform in local space that will be prepended to the LocalLayoutTransform when rendering.
+	 * @param InLocalRenderTransformPivot		Pivot in normalizes local space for the local render transform.
 	 * @param ParentAccumulatedLayoutTransform	The accumulated layout transform of the parent widget. AccumulatedLayoutTransform = Concat(LocalLayoutTransform, ParentAccumulatedLayoutTransform).
 	 * @param ParentAccumulatedRenderTransform	The accumulated render transform of the parent widget. AccumulatedRenderTransform = Concat(LocalRenderTransform, LocalLayoutTransform, ParentAccumulatedRenderTransform).
 	 */
@@ -63,6 +64,7 @@ private:
 		const FVector2D& InLocalSize, 
 		const FSlateLayoutTransform& InLocalLayoutTransform, 
 		const FSlateRenderTransform& InLocalRenderTransform, 
+		const FVector2D& InLocalRenderTransformPivot, 
 		const FSlateLayoutTransform& ParentAccumulatedLayoutTransform, 
 		const FSlateRenderTransform& ParentAccumulatedRenderTransform);
 
@@ -127,10 +129,11 @@ public:
 	 * @param LocalSize			The size of the child geometry in local space.
 	 * @param LayoutTransform	Layout transform of the new child relative to this Geometry. Goes from the child's layout space to the this widget's layout space.
 	 * @param RenderTransform	Render-only transform of the new child that is applied before the layout transform for rendering purposes only.
+	 * @param RenderTransformPivot	Pivot in normalized local space of the Render transform.
 	 *
 	 * @return					The new child geometry.
 	 */
-	FGeometry MakeChild( const FVector2D& LocalSize, const FSlateLayoutTransform& LayoutTransform, const FSlateRenderTransform& RenderTransform ) const;
+	FGeometry MakeChild( const FVector2D& LocalSize, const FSlateLayoutTransform& LayoutTransform, const FSlateRenderTransform& RenderTransform, const FVector2D& RenderTransformPivot ) const;
 
 	/**
 	 * Create a child geometry relative to this one with a given local space size, layout transform, and identity render transform.
diff --git a/Engine/Source/Runtime/SlateCore/Public/Widgets/DeclarativeSyntaxSupport.h b/Engine/Source/Runtime/SlateCore/Public/Widgets/DeclarativeSyntaxSupport.h
--- a/Engine/Source/Runtime/SlateCore/Public/Widgets/DeclarativeSyntaxSupport.h
+++ b/Engine/Source/Runtime/SlateCore/Public/Widgets/DeclarativeSyntaxSupport.h
@@ -1009,7 +1009,8 @@ struct TSlateBaseNamedArgs
 	, _Cursor( TOptional<EMouseCursor::Type>() )
 	, _IsEnabled( true )
 	, _Visibility( EVisibility::Visible )
-	, _RenderTransform()
+	, _RenderTransform( FSlateRenderTransform() )
+	, _RenderTransformPivot( FVector2D::ZeroVector )
 	{
 	}
 
@@ -1048,6 +1049,7 @@ struct TSlateBaseNamedArgs
 	SLATE_ATTRIBUTE( bool, IsEnabled )
 	SLATE_ATTRIBUTE( EVisibility, Visibility )
 	SLATE_ATTRIBUTE( FSlateRenderTransform, RenderTransform )
+	SLATE_ATTRIBUTE( FVector2D, RenderTransformPivot )
 	SLATE_ARGUMENT( FName, Tag )
 
 	TArray<TSharedRef<ISlateMetaData>> MetaData;
@@ -1310,6 +1312,7 @@ struct TDecl
 			InArgs._IsEnabled ,
 			InArgs._Visibility,
 			InArgs._RenderTransform,
+			InArgs._RenderTransformPivot,
 			InArgs._Tag,
 			InArgs.MetaData );
 
diff --git a/Engine/Source/Runtime/SlateCore/Public/Widgets/SWidget.h b/Engine/Source/Runtime/SlateCore/Public/Widgets/SWidget.h
--- a/Engine/Source/Runtime/SlateCore/Public/Widgets/SWidget.h
+++ b/Engine/Source/Runtime/SlateCore/Public/Widgets/SWidget.h
@@ -73,6 +73,7 @@ public:
 		const TAttribute<bool> & InEnabledState ,
 		const TAttribute<EVisibility> & InVisibility,
 		const TAttribute<FSlateRenderTransform>& InTransform,
+		const TAttribute<FVector2D>& InTransformPivot,
 		const FName& InTag,
 		const TArray<TSharedRef<ISlateMetaData>>& InMetaData);
 
@@ -82,10 +83,11 @@ public:
 		const TAttribute<bool> & InEnabledState ,
 		const TAttribute<EVisibility> & InVisibility,
 		const TAttribute<FSlateRenderTransform>& InTransform,
+		const TAttribute<FVector2D>& InTransformPivot,
 		const FName& InTag,
 		const TArray<TSharedRef<ISlateMetaData>>& InMetaData)
 	{
-		Construct(InToolTipText, InToolTip, InCursor, InEnabledState, InVisibility, InTransform, InTag, InMetaData);
+		Construct(InToolTipText, InToolTip, InCursor, InEnabledState, InVisibility, InTransform, InTransformPivot, InTag, InMetaData);
 	}
 
 	//
@@ -565,16 +567,30 @@ public:
 		Visibility = InVisibility;
 	}
 
+	/** @return the render transform of the widget. */
 	const FSlateRenderTransform& GetRenderTransform() const
 	{
 		return RenderTransform.Get();
 	}
 
+	/** @param InTransform the render transform to set for the widget (tranforms from widget's local space). */
 	void SetRenderTransform( TAttribute<FSlateRenderTransform> InTransform )
 	{
 		RenderTransform = InTransform;
 	}
 
+	/** @return the pivot point of the render transform. */
+	const FVector2D& GetRenderTransformPivot() const
+	{
+		return RenderTransformPivot.Get();
+	}
+
+	/** @param InTransformPivot Sets the pivot point of the widget's render transform (in normalized local space). */
+	void SetRenderTransformPivot( TAttribute<FVector2D> InTransformPivot )
+	{
+		RenderTransformPivot = InTransformPivot;
+	}
+
 	/**
 	 * Set the tool tip that should appear when this widget is hovered.
 	 *
@@ -780,6 +796,9 @@ protected:
 
 	/** Render transform of this widget */
 	TAttribute< FSlateRenderTransform > RenderTransform;
+
+	/** Render transform pivot of this widget (in normalized local space) */
+	TAttribute< FVector2D > RenderTransformPivot;
 private:
 
 	/**
