Possible resolution for Mac compilation error regarding FSuperRowType in SMultipleOptionTable's row class widget.

diff --git a/Engine/Source/Runtime/AppFramework/Public/Widgets/SMultipleOptionTable.h b/Engine/Source/Runtime/AppFramework/Public/Widgets/SMultipleOptionTable.h
--- a/Engine/Source/Runtime/AppFramework/Public/Widgets/SMultipleOptionTable.h
+++ b/Engine/Source/Runtime/AppFramework/Public/Widgets/SMultipleOptionTable.h
@@ -3,6 +3,7 @@
 #pragma once
 
 #include "SCheckBox.h"
+#include "SMultiColumnTableRow.h"
 
 #define LOCTEXT_NAMESPACE "SMultipleOptionTable"
 
