Additional compilation solutions for SMultipleOptionTable.h.

diff --git a/Engine/Source/Runtime/AppFramework/Public/Widgets/SMultipleOptionTable.h b/Engine/Source/Runtime/AppFramework/Public/Widgets/SMultipleOptionTable.h
--- a/Engine/Source/Runtime/AppFramework/Public/Widgets/SMultipleOptionTable.h
+++ b/Engine/Source/Runtime/AppFramework/Public/Widgets/SMultipleOptionTable.h
@@ -3,7 +3,7 @@
 #pragma once
 
 #include "SCheckBox.h"
-#include "SMultiColumnTableRow.h"
+#include "STableRow.h"
 
 #define LOCTEXT_NAMESPACE "SMultipleOptionTable"
 
