Fixed Clang compile error with SMultipleOptionTable

diff --git a/Engine/Source/Runtime/AppFramework/Public/Widgets/SMultipleOptionTable.h b/Engine/Source/Runtime/AppFramework/Public/Widgets/SMultipleOptionTable.h
--- a/Engine/Source/Runtime/AppFramework/Public/Widgets/SMultipleOptionTable.h
+++ b/Engine/Source/Runtime/AppFramework/Public/Widgets/SMultipleOptionTable.h
@@ -3,7 +3,6 @@
 #pragma once
 
 #include "SCheckBox.h"
-#include "STableRow.h"
 
 #define LOCTEXT_NAMESPACE "SMultipleOptionTable"
 
@@ -16,6 +15,7 @@ class SOptionTableRow
 {
 public:
 	typedef typename TSlateDelegates< OptionType >::FOnGenerateWidget FOnGenerateWidget;
+	typedef typename STableRow< TSharedPtr<OptionType> >::FArguments FOptionTableRowArgs;
 
 	SLATE_BEGIN_ARGS(SOptionTableRow) { }
 		SLATE_ARGUMENT(TSharedPtr<STableViewBase>, OwnerTableView)
@@ -39,7 +39,7 @@ public:
 		IsChecked = InArgs._IsChecked;
 		OnGenerateWidget= InArgs._OnGenerateWidget;
 
-		SMultiColumnTableRow< TSharedPtr<OptionType> >::Construct(FSuperRowType::FArguments(), InArgs._OwnerTableView.ToSharedRef());
+		SMultiColumnTableRow< TSharedPtr<OptionType> >::Construct(FOptionTableRowArgs(), InArgs._OwnerTableView.ToSharedRef());
 	}
 
 	/**
