Another potential fix for the Mac error in SMultipleOptionTable.h.

diff --git a/Engine/Source/Runtime/AppFramework/Public/Widgets/SMultipleOptionTable.h b/Engine/Source/Runtime/AppFramework/Public/Widgets/SMultipleOptionTable.h
--- a/Engine/Source/Runtime/AppFramework/Public/Widgets/SMultipleOptionTable.h
+++ b/Engine/Source/Runtime/AppFramework/Public/Widgets/SMultipleOptionTable.h
@@ -1,4 +1,4 @@
-// Copyright 1998-2014 Epic Games, Inc. All Rights Reserved.
+﻿// Copyright 1998-2014 Epic Games, Inc. All Rights Reserved.
 
 #pragma once
 
@@ -15,7 +15,7 @@ class SOptionTableRow
 {
 public:
 	typedef typename TSlateDelegates< OptionType >::FOnGenerateWidget FOnGenerateWidget;
-	typedef typename STableRow< TSharedPtr<OptionType> >::FArguments FOptionTableRowArgs;
+	typedef SMultiColumnTableRow< TSharedPtr<OptionType> >::FSuperRowType FSuperRowType;
 
 	SLATE_BEGIN_ARGS(SOptionTableRow) { }
 		SLATE_ARGUMENT(TSharedPtr<STableViewBase>, OwnerTableView)
