Backed out my own change, Fricker's fix seems to work.

diff --git a/Engine/Source/Runtime/AppFramework/Public/Widgets/SMultipleOptionTable.h b/Engine/Source/Runtime/AppFramework/Public/Widgets/SMultipleOptionTable.h
--- a/Engine/Source/Runtime/AppFramework/Public/Widgets/SMultipleOptionTable.h
+++ b/Engine/Source/Runtime/AppFramework/Public/Widgets/SMultipleOptionTable.h
@@ -1,4 +1,4 @@
-﻿// Copyright 1998-2014 Epic Games, Inc. All Rights Reserved.
+// Copyright 1998-2014 Epic Games, Inc. All Rights Reserved.
 
 #pragma once
 
@@ -15,7 +15,7 @@ class SOptionTableRow
 {
 public:
 	typedef typename TSlateDelegates< OptionType >::FOnGenerateWidget FOnGenerateWidget;
-	typedef SMultiColumnTableRow< TSharedPtr<OptionType> >::FSuperRowType FSuperRowType;
+	typedef typename STableRow< TSharedPtr<OptionType> >::FArguments FOptionTableRowArgs;
 
 	SLATE_BEGIN_ARGS(SOptionTableRow) { }
 		SLATE_ARGUMENT(TSharedPtr<STableViewBase>, OwnerTableView)
