Corrected conflicting LOCTEXT keys (RestoreExpansionButton) in STableViewTesting.cpp.

diff --git a/Engine/Source/Runtime/AppFramework/Private/Framework/Testing/STableViewTesting.cpp b/Engine/Source/Runtime/AppFramework/Private/Framework/Testing/STableViewTesting.cpp
--- a/Engine/Source/Runtime/AppFramework/Private/Framework/Testing/STableViewTesting.cpp
+++ b/Engine/Source/Runtime/AppFramework/Private/Framework/Testing/STableViewTesting.cpp
@@ -661,7 +661,7 @@ public:
 				+SHorizontalBox::Slot() .AutoWidth() .VAlign(VAlign_Center)
 				[
 					SNew(SButton)
-					.Text( LOCTEXT("RestoreExpansionsButton", "Collapse All") )
+					.Text( LOCTEXT("CollapseExpansionsButton", "Collapse All") )
 					.OnClicked( this, &STableViewTesting::CollapseAll )
 				]
 				+SHorizontalBox::Slot() .AutoWidth() .VAlign(VAlign_Center)
