Editor: Workaround for color picker themes crash on shutdown

The problem is that the color picker has a global static instance of the color themes widget. By the time that widget is being destroyed as part of static shutdown, Slate and ICU are already unloaded. Ideally, we wouldn't have a global variable here at all. Some further research is needed to understand why it was added in the first place. In the meantime, the variable is being changed to a weak pointer, so it gets destroyed when the last color picker instance gets destroyed.

diff --git a/Engine/Source/Runtime/AppFramework/Public/Widgets/Colors/SColorPicker.cpp b/Engine/Source/Runtime/AppFramework/Public/Widgets/Colors/SColorPicker.cpp
--- a/Engine/Source/Runtime/AppFramework/Public/Widgets/Colors/SColorPicker.cpp
+++ b/Engine/Source/Runtime/AppFramework/Public/Widgets/Colors/SColorPicker.cpp
@@ -16,7 +16,7 @@ const FVector2D SColorPicker::DEFAULT_WINDOW_SIZE = FVector2D(308, 458);
 /** The max time allowed for updating before we shut off auto-updating */
 const double SColorPicker::MAX_ALLOWED_UPDATE_TIME = 0.1;
 
-TSharedPtr<SColorThemesViewer> SColorPicker::ColorThemesViewer;
+TWeakPtr<SColorThemesViewer> SColorPicker::ColorThemesViewer;
 
 
 /* SColorPicker structors
@@ -24,9 +24,11 @@ TSharedPtr<SColorThemesViewer> SColorPicker::ColorThemesViewer;
 
 SColorPicker::~SColorPicker()
 {
-	if (ColorThemesViewer.IsValid())
+	TSharedPtr<SColorThemesViewer> ThemesViewer = ColorThemesViewer.Pin();
+
+	if (ThemesViewer.IsValid())
 	{
-		ColorThemesViewer->OnCurrentThemeChanged().RemoveAll(this);
+		ThemesViewer->OnCurrentThemeChanged().RemoveAll(this);
 	}
 }
 
@@ -207,14 +209,17 @@ void SColorPicker::GenerateDefaultColorPickerContent( bool bAdvancedSectionExpan
 	// The height of the gradient bars beneath the sliders
 	const FSlateFontInfo SmallLayoutFont = FCoreStyle::Get().GetFontStyle("ColorPicker.Font");
 
-	if (!ColorThemesViewer.IsValid())
+	TSharedPtr<SColorThemesViewer> ThemesViewer = ColorThemesViewer.Pin();
+
+	if (!ThemesViewer.IsValid())
 	{
-		ColorThemesViewer = SNew(SColorThemesViewer);
+		ThemesViewer = SNew(SColorThemesViewer);
+		ColorThemesViewer = ThemesViewer;
 	}
 
-	ColorThemesViewer->OnCurrentThemeChanged().AddSP(this, &SColorPicker::HandleThemesViewerThemeChanged);
-	ColorThemesViewer->SetUseAlpha(bUseAlpha);
-	ColorThemesViewer->MenuToStandardNoReturn();
+	ThemesViewer->OnCurrentThemeChanged().AddSP(this, &SColorPicker::HandleThemesViewerThemeChanged);
+	ThemesViewer->SetUseAlpha(bUseAlpha);
+	ThemesViewer->MenuToStandardNoReturn();
 
 	// The standard button to open the color themes can temporarily become a trash for colors
 	ColorThemeComboButton = SNew(SComboButton)
@@ -222,7 +227,7 @@ void SColorPicker::GenerateDefaultColorPickerContent( bool bAdvancedSectionExpan
 		.MenuPlacement(MenuPlacement_ComboBox)
 		.ToolTipText(LOCTEXT("OpenThemeManagerToolTip", "Open Color Theme Manager"));
 
-	ColorThemeComboButton->SetMenuContent(ColorThemesViewer.ToSharedRef());
+	ColorThemeComboButton->SetMenuContent(ThemesViewer.ToSharedRef());
 
 	SmallTrash = SNew(SHorizontalBox)
 	+ SHorizontalBox::Slot()
@@ -1438,17 +1443,29 @@ void SColorPicker::HandleThemeBarColorSelected( FLinearColor NewValue )
 
 TSharedPtr<FColorTheme> SColorPicker::HandleThemeBarColorTheme() const
 {
-	return ColorThemesViewer->GetCurrentColorTheme();
+	TSharedPtr<SColorThemesViewer> ThemesViewer = ColorThemesViewer.Pin();
+
+	if (ThemesViewer.IsValid())
+	{
+		return ThemesViewer->GetCurrentColorTheme();
+	}
+
+	return nullptr;
 }
 
 
 EVisibility SColorPicker::HandleThemeBarHintVisibility( ) const
 {
-	TSharedPtr<FColorTheme> SelectedTheme = ColorThemesViewer->GetCurrentColorTheme();
+	TSharedPtr<SColorThemesViewer> ThemesViewer = ColorThemesViewer.Pin();
 
-	if (SelectedTheme.IsValid() && SelectedTheme->GetColors().Num() == 0)
+	if (ThemesViewer.IsValid())
 	{
-		return EVisibility::HitTestInvisible;
+		TSharedPtr<FColorTheme> SelectedTheme = ThemesViewer->GetCurrentColorTheme();
+
+		if (SelectedTheme.IsValid() && SelectedTheme->GetColors().Num() == 0)
+		{
+			return EVisibility::HitTestInvisible;
+		}
 	}
 
 	return EVisibility::Hidden;
diff --git a/Engine/Source/Runtime/AppFramework/Public/Widgets/Colors/SColorPicker.h b/Engine/Source/Runtime/AppFramework/Public/Widgets/Colors/SColorPicker.h
--- a/Engine/Source/Runtime/AppFramework/Public/Widgets/Colors/SColorPicker.h
+++ b/Engine/Source/Runtime/AppFramework/Public/Widgets/Colors/SColorPicker.h
@@ -427,7 +427,7 @@ private:
 private:
 
 	/** A static pointer to the global color themes viewer */
-	static TSharedPtr<class SColorThemesViewer> ColorThemesViewer;
+	static TWeakPtr<class SColorThemesViewer> ColorThemesViewer;
 };
 
 
