Fixed text justification not working correctly when the UI was scaled

The width used for justification is the larger of the "desired size" (as fixed wrapping widths force the desired size to that width), or the "actual size" (in the case that the text is smaller than its containing widget, and the text widget is set to fill width).

The "actual size" wasn't being scaled before being used to offset the text blocks, and also wasn't taking into account the text layout margin. Both of these issues lead to some incorrect offsets being calculated.

Also added a test case for STextBlock justification to the widget gallery in Slate Viewer.

diff --git a/Engine/Source/Runtime/AppFramework/Private/Framework/Testing/SWidgetGallery.cpp b/Engine/Source/Runtime/AppFramework/Private/Framework/Testing/SWidgetGallery.cpp
--- a/Engine/Source/Runtime/AppFramework/Private/Framework/Testing/SWidgetGallery.cpp
+++ b/Engine/Source/Runtime/AppFramework/Private/Framework/Testing/SWidgetGallery.cpp
@@ -657,8 +657,31 @@ public:
 					+ SGridPanel::Slot(1, 25)
 						.Padding(0.0f, 5.0f)
 						[
-							SNew(STextBlock)
-								.Text(LOCTEXT("STextBlockExampleLabel", "This is a text block"))			
+							SNew(SVerticalBox)
+							
+							+ SVerticalBox::Slot()
+							.AutoHeight()
+							[
+								SNew(STextBlock)
+									.Text(LOCTEXT("STextBlockLeftExampleLabel", "This is a left-justified text block"))
+									.Justification(ETextJustify::Left)
+							]
+
+							+ SVerticalBox::Slot()
+							.AutoHeight()
+							[
+								SNew(STextBlock)
+									.Text(LOCTEXT("STextBlockCenterExampleLabel", "This is a center-justified text block"))
+									.Justification(ETextJustify::Center)
+							]
+
+							+ SVerticalBox::Slot()
+							.AutoHeight()
+							[
+								SNew(STextBlock)
+									.Text(LOCTEXT("STextBlockRightExampleLabel", "This is a right-justified text block"))
+									.Justification(ETextJustify::Right)
+							]
 						]
 
 					// STextComboBox
diff --git a/Engine/Source/Runtime/Slate/Private/Framework/Text/TextLayout.cpp b/Engine/Source/Runtime/Slate/Private/Framework/Text/TextLayout.cpp
--- a/Engine/Source/Runtime/Slate/Private/Framework/Text/TextLayout.cpp
+++ b/Engine/Source/Runtime/Slate/Private/Framework/Text/TextLayout.cpp
@@ -265,7 +265,7 @@ void FTextLayout::JustifyLayout()
 		return;
 	}
 
-	const float LayoutWidthNoMargin = FMath::Max(DrawSize.X - Margin.GetTotalSpaceAlong<Orient_Horizontal>() * Scale, ViewSize.X);
+	const float LayoutWidthNoMargin = FMath::Max(DrawSize.X, ViewSize.X * Scale) - (Margin.GetTotalSpaceAlong<Orient_Horizontal>() * Scale);
 
 	for (FLineView& LineView : LineViews)
 	{
