Material Editor now opens Color Picker with Advanced section automatically expanded.
#ttp 345475 - NUX:Material Editor: Vector parameter alpha channel is in the "advanced" part of the color picker.
#branch UE4
#proj Editor.MaterialEditor, Runtime.AppFramework
#add Added extra construction parameter to SColorPicker (ExpandAdvancedSection) and added it as an option to SColorPicker::OpenColorPicker.

diff --git a/Engine/Source/Editor/MaterialEditor/Private/MaterialEditor.cpp b/Engine/Source/Editor/MaterialEditor/Private/MaterialEditor.cpp
--- a/Engine/Source/Editor/MaterialEditor/Private/MaterialEditor.cpp
+++ b/Engine/Source/Editor/MaterialEditor/Private/MaterialEditor.cpp
@@ -3808,6 +3808,7 @@ void FMaterialEditor::OnNodeDoubleClicked(class UEdGraphNode* Node)
 			PickerArgs.ParentWidget = GraphEditor;//AsShared();
 			PickerArgs.bUseAlpha = Constant4Expression != NULL || VectorExpression != NULL;
 			PickerArgs.bOnlyRefreshOnOk = true;
+			PickerArgs.bExpandAdvancedSection = true;
 			PickerArgs.DisplayGamma = TAttribute<float>::Create( TAttribute<float>::FGetter::CreateUObject(GEngine, &UEngine::GetDisplayGamma) );
 			PickerArgs.ColorChannelsArray = &Channels;
 			PickerArgs.OnColorCommitted = FOnLinearColorValueChanged::CreateSP(this, &FMaterialEditor::OnColorPickerCommitted);
diff --git a/Engine/Source/Runtime/AppFramework/Public/Widgets/Colors/SColorPicker.cpp b/Engine/Source/Runtime/AppFramework/Public/Widgets/Colors/SColorPicker.cpp
--- a/Engine/Source/Runtime/AppFramework/Public/Widgets/Colors/SColorPicker.cpp
+++ b/Engine/Source/Runtime/AppFramework/Public/Widgets/Colors/SColorPicker.cpp
@@ -74,7 +74,7 @@ void SColorPicker::Construct( const FArguments& InArgs )
 
 	bool bAdvancedSectionExpanded = false;
 
-	if (FPaths::FileExists(GEditorUserSettingsIni))
+	if (!FPaths::FileExists(GEditorUserSettingsIni))
 	{
 		bool WheelMode = true;
 
@@ -85,6 +85,8 @@ void SColorPicker::Construct( const FArguments& InArgs )
 		CurrentMode = WheelMode ? EColorPickerModes::Wheel : EColorPickerModes::Spectrum;
 	}
 
+	bAdvancedSectionExpanded |= InArgs._ExpandAdvancedSection;
+
 	if (bColorPickerIsInlineVersion)
 	{
 		GenerateInlineColorPickerContent();
@@ -1546,6 +1548,7 @@ bool OpenColorPicker(const FColorPickerArgs& Args)
 		.TargetLinearColors(Args.LinearColorArray ? *Args.LinearColorArray : TArray<FLinearColor*>())
 		.TargetColorChannels(Args.ColorChannelsArray ? *Args.ColorChannelsArray : TArray<FColorChannels>())
 		.UseAlpha(Args.bUseAlpha)
+		.ExpandAdvancedSection(Args.bExpandAdvancedSection)
 		.OnlyRefreshOnMouseUp(Args.bOnlyRefreshOnMouseUp && !Args.bIsModal)
 		.OnlyRefreshOnOk(Args.bOnlyRefreshOnOk || Args.bIsModal)
 		.OnColorCommitted(Args.OnColorCommitted)
diff --git a/Engine/Source/Runtime/AppFramework/Public/Widgets/Colors/SColorPicker.h b/Engine/Source/Runtime/AppFramework/Public/Widgets/Colors/SColorPicker.h
--- a/Engine/Source/Runtime/AppFramework/Public/Widgets/Colors/SColorPicker.h
+++ b/Engine/Source/Runtime/AppFramework/Public/Widgets/Colors/SColorPicker.h
@@ -129,6 +129,9 @@ public:
 		/** If true, this color picker will be a stripped down version of the full color picker */
 		SLATE_ARGUMENT(bool, DisplayInlineVersion)
 
+		/** If true, the Advanced section will be expanded, regardless of the remembered state */
+		SLATE_ARGUMENT(bool, ExpandAdvancedSection)
+
 	SLATE_END_ARGS()
 	
 	/** A default window size for the color picker which looks nice */
@@ -448,6 +451,9 @@ struct FColorPickerArgs
 	/** Whether to disable the refresh until the picker closes. */
 	bool bOnlyRefreshOnOk;
 
+	/** Whether to automatically expand the Advanced section. */
+	bool bExpandAdvancedSection;
+
 	/** The current display gamma used to correct colors picked from the display. */
 	TAttribute<float> DisplayGamma;
 
@@ -487,6 +493,7 @@ struct FColorPickerArgs
 		, bUseAlpha(false)
 		, bOnlyRefreshOnMouseUp(false)
 		, bOnlyRefreshOnOk(false)
+		, bExpandAdvancedSection(false)
 		, DisplayGamma(2.2f)
 		, ColorArray(nullptr)
 		, LinearColorArray(nullptr)
