Fix crash when opening tutorial on Linux.

- Caused by class name collision: we currently export all the names (i.e. visibility=default) in non-Shipping configurations because this is needed for callstack generation. If there are two (or more) classes named the same in different dynamic libraries, this can cause problems during run-time, because the names will be resolved pretty much arbitrarily (depending on the order of loading the said libraries)
- This particular issue seems to have involved FTextStyleDecorator::Create() static method, which, instead of creating an instance of FTextStyleDecorator from IntroTutorials module, created FTextStyleDecorator from the test suite.
- Solution for this one is to rename one of classes to something different.

diff --git a/Engine/Source/Runtime/AppFramework/Private/Widgets/Testing/STestSuite.cpp b/Engine/Source/Runtime/AppFramework/Private/Widgets/Testing/STestSuite.cpp
--- a/Engine/Source/Runtime/AppFramework/Private/Widgets/Testing/STestSuite.cpp
+++ b/Engine/Source/Runtime/AppFramework/Private/Widgets/Testing/STestSuite.cpp
@@ -2289,16 +2289,16 @@ struct FTextStyles
  * This is a custom decorator used to allow arbitrary styling of text within a rich-text editor
  * This is required since normal text styling can only work with known styles from a given Slate style-set
  */
-class FTextStyleDecorator : public ITextDecorator
+class FTextStyleDecoratorTest : public ITextDecorator
 {
 public:
 
-	static TSharedRef<FTextStyleDecorator> Create(FTextStyles* const InTextStyles)
+	static TSharedRef<FTextStyleDecoratorTest> Create(FTextStyles* const InTextStyles)
 	{
-		return MakeShareable(new FTextStyleDecorator(InTextStyles));
+		return MakeShareable(new FTextStyleDecoratorTest(InTextStyles));
 	}
 
-	virtual ~FTextStyleDecorator()
+	virtual ~FTextStyleDecoratorTest()
 	{
 	}
 
@@ -2325,7 +2325,7 @@ public:
 
 private:
 
-	FTextStyleDecorator(FTextStyles* const InTextStyles)
+	FTextStyleDecoratorTest(FTextStyles* const InTextStyles)
 		: TextStyles(InTextStyles)
 	{
 	}
@@ -2388,7 +2388,7 @@ public:
 		// ... so we also need to add some decorators to handle the things we want to demo
 		OnHyperlinkClicked = FSlateHyperlinkRun::FOnClick::CreateStatic(&RichTextHelper::OnBrowserLinkClicked, AsShared());
 		RichTextMarshaller->AppendInlineDecorator(FHyperlinkDecorator::Create(TEXT("browser"), OnHyperlinkClicked));
-		RichTextMarshaller->AppendInlineDecorator(FTextStyleDecorator::Create(&TextStyles));
+		RichTextMarshaller->AppendInlineDecorator(FTextStyleDecoratorTest::Create(&TextStyles));
 
 		// The syntax highlighter marshaller is self contained, so doesn't need any extra configuration
 		SyntaxHighlighterMarshaller = FRichTextSyntaxHighlighterTextLayoutMarshaller::Create(
