Editor: Fixed color picker hex input resulting in gamma corrected value

Based on GitHub PR# 280 by Fienzine (https://github.com/EpicGames/UnrealEngine/pull/280)

diff --git a/Engine/Source/Runtime/AppFramework/Public/Widgets/Colors/SColorPicker.cpp b/Engine/Source/Runtime/AppFramework/Public/Widgets/Colors/SColorPicker.cpp
--- a/Engine/Source/Runtime/AppFramework/Public/Widgets/Colors/SColorPicker.cpp
+++ b/Engine/Source/Runtime/AppFramework/Public/Widgets/Colors/SColorPicker.cpp
@@ -1259,7 +1259,8 @@ void SColorPicker::HandleHexInputTextCommitted( const FText& Text, ETextCommit::
 {
 	if ((CommitType == ETextCommit::OnEnter) || (CommitType == ETextCommit::OnUserMovedFocus))
 	{
-		SetNewTargetColorRGB(FColor::FromHex(Text.ToString()), false);
+		FColor Color = FColor::FromHex(Text.ToString());
+		SetNewTargetColorRGB(FLinearColor(Color.R / 255.0f, Color.G / 255.0f, Color.B / 255.0f, Color.A / 255.0f), false);
 	}	
 }
 
