Slate enabled test case now includes render transforms.

diff --git a/Engine/Source/Runtime/AppFramework/Private/Widgets/Testing/STestSuite.cpp b/Engine/Source/Runtime/AppFramework/Private/Widgets/Testing/STestSuite.cpp
--- a/Engine/Source/Runtime/AppFramework/Private/Widgets/Testing/STestSuite.cpp
+++ b/Engine/Source/Runtime/AppFramework/Private/Widgets/Testing/STestSuite.cpp
@@ -728,17 +728,25 @@ public:
 				.ContentPadding(20).HAlign(HAlign_Center).VAlign(VAlign_Center)
 				[
 					SNew(SButton)
-					.IsEnabled( false )
+					.RenderTransform( FSlateRenderTransform( Concatenate( FVector2D(20.0f, 20.0f), FQuat2D(-PI/12.0f) ) ) )
 					.ContentPadding(20).HAlign(HAlign_Center).VAlign(VAlign_Center)
 					[
 						SNew(SButton)
+						.RenderTransform( FSlateRenderTransform( Concatenate( FVector2D(20.0f, 20.0f), FQuat2D(-PI/12.0f) ) ) )
 						.ContentPadding(20).HAlign(HAlign_Center).VAlign(VAlign_Center)
 						[
 							SNew(SButton)
+							.IsEnabled(false)
+							.RenderTransform( FSlateRenderTransform( Concatenate( FVector2D(20.0f, 20.0f), FQuat2D(-PI/12.0f) ) ) )
 							.ContentPadding(20).HAlign(HAlign_Center).VAlign(VAlign_Center)
 							[
-								SNew(STextBlock)
-								.Text(NSLOCTEXT("x","ClickMe","ClickMe!"))
+								SNew(SButton)
+								.RenderTransform( FSlateRenderTransform( Concatenate( FVector2D(20.0f, 20.0f), FQuat2D(-PI/12.0f) ) ) )
+								.ContentPadding(20).HAlign(HAlign_Center).VAlign(VAlign_Center)
+								[
+									SNew(STextBlock)
+									.Text(NSLOCTEXT("x","ClickMe","ClickMe!"))
+								]
 							]
 						]
 					]
