Moving Web Browser widget from Slate to App Framework

diff --git a/Engine/Source/Runtime/AppFramework/AppFramework.Build.cs b/Engine/Source/Runtime/AppFramework/AppFramework.Build.cs
--- a/Engine/Source/Runtime/AppFramework/AppFramework.Build.cs
+++ b/Engine/Source/Runtime/AppFramework/AppFramework.Build.cs
@@ -16,7 +16,8 @@ public class AppFramework : ModuleRules
 
 		PrivateDependencyModuleNames.AddRange(
 			new string[] {
-                "InputCore",
+				"InputCore",
+				"WebBrowser",
 			}
 		);
 
diff --git a/Engine/Source/Runtime/AppFramework/Private/Widgets/SWebBrowser.cpp b/Engine/Source/Runtime/AppFramework/Private/Widgets/SWebBrowser.cpp
new file mode 100644
--- /dev/null
+++ b/Engine/Source/Runtime/AppFramework/Private/Widgets/SWebBrowser.cpp
@@ -0,0 +1,190 @@
+// Copyright 1998-2014 Epic Games, Inc. All Rights Reserved.
+
+#include "AppFrameworkPrivatePCH.h"
+#include "SWebBrowser.h"
+
+#if WITH_EDITOR || IS_PROGRAM
+#include "WebBrowserModule.h"
+#include "IWebBrowserSingleton.h"
+#include "IWebBrowserWindow.h"
+#include "WebBrowserViewport.h"
+#endif
+
+#define LOCTEXT_NAMESPACE "WebBrowser"
+
+SWebBrowser::SWebBrowser()
+{
+}
+
+void SWebBrowser::Construct(const FArguments& InArgs)
+{
+	void* OSWindowHandle = nullptr;
+	if (InArgs._ParentWindow.IsValid())
+	{
+		TSharedPtr<FGenericWindow> NativeWindow = InArgs._ParentWindow->GetNativeWindow();
+		OSWindowHandle = NativeWindow->GetOSWindowHandle();
+	}
+
+#if WITH_EDITOR || IS_PROGRAM
+	BrowserWindow = IWebBrowserModule::Get().GetSingleton()->CreateBrowserWindow(OSWindowHandle,
+	                                                                             InArgs._InitialURL,
+	                                                                             InArgs._ViewportSize.Get().X,
+	                                                                             InArgs._ViewportSize.Get().Y,
+	                                                                             InArgs._SupportsTransparency);
+#endif
+
+	TSharedPtr<SViewport> ViewportWidget;
+
+	ChildSlot
+	[
+		SNew(SVerticalBox)
+		+SVerticalBox::Slot()
+		.AutoHeight()
+		.Padding(0, 5)
+		[
+			SNew(SHorizontalBox)
+			.Visibility(InArgs._ShowControls ? EVisibility::Visible : EVisibility::Collapsed)
+			+SHorizontalBox::Slot()
+			.AutoWidth()
+			[
+				SNew(SButton)
+				.Text(LOCTEXT("Back","Back"))
+				.IsEnabled(this, &SWebBrowser::CanGoBack)
+				.OnClicked(this, &SWebBrowser::OnBackClicked)
+			]
+			+SHorizontalBox::Slot()
+			.AutoWidth()
+			[
+				SNew(SButton)
+				.Text(LOCTEXT("Forward", "Forward"))
+				.IsEnabled(this, &SWebBrowser::CanGoForward)
+				.OnClicked(this, &SWebBrowser::OnForwardClicked)
+			]
+			+SHorizontalBox::Slot()
+			.AutoWidth()
+			[
+				SNew(SButton)
+				.Text(this, &SWebBrowser::GetReloadButtonText)
+				.OnClicked(this, &SWebBrowser::OnReloadClicked)
+			]
+			+SHorizontalBox::Slot()
+			.FillWidth(1.0f)
+			.VAlign(VAlign_Center)
+			.HAlign(HAlign_Right)
+			.Padding(5)
+			[
+				SNew(STextBlock)
+				.Text(this, &SWebBrowser::GetTitleText)
+				.Justification(ETextJustify::Right)
+			]
+		]
+		+SVerticalBox::Slot()
+		[
+			SAssignNew(ViewportWidget, SViewport)
+			.ViewportSize(InArgs._ViewportSize)
+			.EnableGammaCorrection(false)
+			.EnableBlending(InArgs._SupportsTransparency)
+			.IgnoreTextureAlpha(!InArgs._SupportsTransparency)
+		]
+	];
+
+#if WITH_EDITOR || IS_PROGRAM
+	if (BrowserWindow.IsValid())
+	{
+		BrowserViewport = MakeShareable(new FWebBrowserViewport(BrowserWindow, ViewportWidget));
+		ViewportWidget->SetViewportInterface(BrowserViewport.ToSharedRef());
+	}
+#endif
+}
+
+FText SWebBrowser::GetTitleText() const
+{
+#if WITH_EDITOR || IS_PROGRAM
+	if (BrowserWindow.IsValid())
+	{
+		return FText::FromString(BrowserWindow->GetTitle());
+	}
+#endif
+	return LOCTEXT("InvalidWindow", "Browser Window is not valid/supported");
+}
+
+bool SWebBrowser::CanGoBack() const
+{
+#if WITH_EDITOR || IS_PROGRAM
+	if (BrowserWindow.IsValid())
+	{
+		return BrowserWindow->CanGoBack();
+	}
+#endif
+	return false;
+}
+
+FReply SWebBrowser::OnBackClicked()
+{
+#if WITH_EDITOR || IS_PROGRAM
+	if (BrowserWindow.IsValid())
+	{
+		BrowserWindow->GoBack();
+	}
+#endif
+	return FReply::Handled();
+}
+
+bool SWebBrowser::CanGoForward() const
+{
+#if WITH_EDITOR || IS_PROGRAM
+	if (BrowserWindow.IsValid())
+	{
+		return BrowserWindow->CanGoForward();
+	}
+#endif
+	return false;
+}
+
+FReply SWebBrowser::OnForwardClicked()
+{
+#if WITH_EDITOR || IS_PROGRAM
+	if (BrowserWindow.IsValid())
+	{
+		BrowserWindow->GoForward();
+	}
+#endif
+	return FReply::Handled();
+}
+
+FText SWebBrowser::GetReloadButtonText() const
+{
+	static FText Reload = LOCTEXT("Reload", "Reload");
+	static FText Stop = LOCTEXT("Stop", "Stop");
+
+#if WITH_EDITOR || IS_PROGRAM
+	if (BrowserWindow.IsValid())
+	{
+		if (BrowserWindow->IsLoading())
+		{
+			return Stop;
+		}
+	}
+#endif
+	return Reload;
+}
+
+FReply SWebBrowser::OnReloadClicked()
+{
+#if WITH_EDITOR || IS_PROGRAM
+	if (BrowserWindow.IsValid())
+	{
+		if (BrowserWindow->IsLoading())
+		{
+			BrowserWindow->StopLoad();
+		}
+		else
+		{
+			BrowserWindow->Reload();
+		}
+	}
+#endif
+	return FReply::Handled();
+}
+
+#undef LOCTEXT_NAMESPACE
diff --git a/Engine/Source/Runtime/AppFramework/Public/Widgets/SWebBrowser.h b/Engine/Source/Runtime/AppFramework/Public/Widgets/SWebBrowser.h
new file mode 100644
--- /dev/null
+++ b/Engine/Source/Runtime/AppFramework/Public/Widgets/SWebBrowser.h
@@ -0,0 +1,88 @@
+// Copyright 1998-2014 Epic Games, Inc. All Rights Reserved.
+
+#pragma once
+
+class IWebBrowserWindow;
+class FWebBrowserViewport;
+
+class APPFRAMEWORK_API SWebBrowser : public SCompoundWidget
+{
+public:
+	SLATE_BEGIN_ARGS(SWebBrowser)
+		: _InitialURL(TEXT("www.google.com"))
+		, _ShowControls(true)
+		, _SupportsTransparency(false)
+		, _ViewportSize(FVector2D(320, 240))
+		{}
+
+		/** A reference to the parent window */
+		SLATE_ARGUMENT(TSharedPtr<SWindow>, ParentWindow)
+
+		/** URL that the browser will initially navigate to */
+		SLATE_ARGUMENT(FString, InitialURL)
+
+		/** Whether to show standard controls like Back, Forward, Reload etc. */
+		SLATE_ARGUMENT(bool, ShowControls)
+
+		/** Should this browser window support transparency */
+		SLATE_ARGUMENT(bool, SupportsTransparency)
+
+		/** Desired size of the web browser viewport */
+		SLATE_ATTRIBUTE(FVector2D, ViewportSize);
+	SLATE_END_ARGS()
+
+	/**
+	 * Constructor
+	 */
+	SWebBrowser();
+
+	/**
+	 * Construct the widget.
+	 *
+	 * @param InArgs  Declaration from which to construct the widget.
+	 */
+	void Construct(const FArguments& InArgs);
+
+	/**
+	 * Get the current title of the web page
+	 */
+	FText GetTitleText() const;
+
+private:
+	/**
+	 * Returns true if the browser can navigate backwards.
+	 */
+	bool CanGoBack() const;
+
+	/**
+	 * Navigate backwards.
+	 */
+	FReply OnBackClicked();
+
+	/**
+	 * Returns true if the browser can navigate forwards.
+	 */
+	bool CanGoForward() const;
+
+	/**
+	 * Navigate forwards.
+	 */
+	FReply OnForwardClicked();
+
+	/**
+	 * Get text for reload button depending on status
+	 */
+	FText GetReloadButtonText() const;
+
+	/**
+	 * Reload or stop loading
+	 */
+	FReply OnReloadClicked();
+
+private:
+	/** Interface for dealing with a web browser window */
+	TSharedPtr<IWebBrowserWindow>	BrowserWindow;
+
+	/** Viewport interface for rendering the web page */
+	TSharedPtr<FWebBrowserViewport>	BrowserViewport;
+};
diff --git a/Engine/Source/Runtime/Slate/Private/Widgets/SWebBrowser.cpp b/Engine/Source/Runtime/Slate/Private/Widgets/SWebBrowser.cpp
deleted file mode 100644
diff --git a/Engine/Source/Runtime/Slate/Public/Widgets/SWebBrowser.h b/Engine/Source/Runtime/Slate/Public/Widgets/SWebBrowser.h
deleted file mode 100644
