Always show advanced section for color pickers

diff --git a/Engine/Source/Runtime/AppFramework/Public/Widgets/Colors/SColorPicker.h b/Engine/Source/Runtime/AppFramework/Public/Widgets/Colors/SColorPicker.h
--- a/Engine/Source/Runtime/AppFramework/Public/Widgets/Colors/SColorPicker.h
+++ b/Engine/Source/Runtime/AppFramework/Public/Widgets/Colors/SColorPicker.h
@@ -76,6 +76,7 @@ public:
 		, _ParentWindow()
 		, _DisplayGamma(2.2f)
 		, _DisplayInlineVersion(false)
+		, _ExpandAdvancedSection(false)
 	{ }
 		
 		/** The color that is being targeted as a TAttribute */
@@ -493,7 +494,7 @@ struct FColorPickerArgs
 		, bUseAlpha(false)
 		, bOnlyRefreshOnMouseUp(false)
 		, bOnlyRefreshOnOk(false)
-		, bExpandAdvancedSection(false)
+		, bExpandAdvancedSection(true)
 		, DisplayGamma(2.2f)
 		, ColorArray(nullptr)
 		, LinearColorArray(nullptr)
