AppFramework: Moving towards "include what you use"; moved files into correct folder; some formatting fixes

diff --git a/Engine/Source/Runtime/AppFramework/Private/AppFrameworkModule.cpp b/Engine/Source/Runtime/AppFramework/Private/AppFrameworkModule.cpp
--- a/Engine/Source/Runtime/AppFramework/Private/AppFrameworkModule.cpp
+++ b/Engine/Source/Runtime/AppFramework/Private/AppFrameworkModule.cpp
@@ -1,6 +1,7 @@
 // Copyright 1998-2014 Epic Games, Inc. All Rights Reserved.
 
 #include "AppFrameworkPrivatePCH.h"
+#include "ModuleInterface.h"
 
 
 static const FName AppFrameworkTabName("AppFramework");
@@ -16,11 +17,11 @@ public:
 
 	// IModuleInterface interface
 	
-	virtual void StartupModule( ) override
+	virtual void StartupModule() override
 	{
 	}
 
-	virtual void ShutdownModule( ) override
+	virtual void ShutdownModule() override
 	{
 	}
 };
diff --git a/Engine/Source/Runtime/AppFramework/Private/AppFrameworkPrivatePCH.h b/Engine/Source/Runtime/AppFramework/Private/AppFrameworkPrivatePCH.h
--- a/Engine/Source/Runtime/AppFramework/Private/AppFrameworkPrivatePCH.h
+++ b/Engine/Source/Runtime/AppFramework/Private/AppFrameworkPrivatePCH.h
@@ -2,12 +2,13 @@
 
 #pragma once
 
-#include "AppFramework.h"
-
 
 /* Private dependencies
  *****************************************************************************/
 
+#include "Core.h"
+#include "SlateBasics.h"
+
 
 /* Private includes
  *****************************************************************************/
diff --git a/Engine/Source/Runtime/AppFramework/Private/Widgets/Colors/SColorPicker.cpp b/Engine/Source/Runtime/AppFramework/Private/Widgets/Colors/SColorPicker.cpp
new file mode 100644
--- /dev/null
+++ b/Engine/Source/Runtime/AppFramework/Private/Widgets/Colors/SColorPicker.cpp
@@ -0,0 +1,1626 @@
+// Copyright 1998-2014 Epic Games, Inc. All Rights Reserved.
+
+#include "AppFrameworkPrivatePCH.h"
+#include "SComplexGradient.h"
+#include "SSimpleGradient.h"
+#include "SColorValueSlider.h"
+#include "SEyeDropperButton.h"
+#include "SColorPicker.h"
+#include "SColorWheel.h"
+#include "SColorSpectrum.h"
+#include "SColorThemes.h"
+#include "SExpandableArea.h"
+
+
+#define LOCTEXT_NAMESPACE "ColorPicker"
+
+
+/** A default window size for the color picker which looks nice */
+const FVector2D SColorPicker::DEFAULT_WINDOW_SIZE = FVector2D(308, 458);
+
+/** The max time allowed for updating before we shut off auto-updating */
+const double SColorPicker::MAX_ALLOWED_UPDATE_TIME = 0.1;
+
+TWeakPtr<SColorThemesViewer> SColorPicker::ColorThemesViewer;
+
+
+/* SColorPicker structors
+ *****************************************************************************/
+
+SColorPicker::~SColorPicker()
+{
+	TSharedPtr<SColorThemesViewer> ThemesViewer = ColorThemesViewer.Pin();
+
+	if (ThemesViewer.IsValid())
+	{
+		ThemesViewer->OnCurrentThemeChanged().RemoveAll(this);
+	}
+}
+
+
+/* SColorPicker methods
+ *****************************************************************************/
+
+void SColorPicker::Construct( const FArguments& InArgs )
+{
+	TargetColorAttribute = InArgs._TargetColorAttribute;
+	CurrentColorHSV = OldColor = TargetColorAttribute.Get().LinearRGBToHSV();
+	CurrentColorRGB = TargetColorAttribute.Get();
+	CurrentMode = EColorPickerModes::Wheel;
+	TargetFColors = InArgs._TargetFColors.Get();
+	TargetLinearColors = InArgs._TargetLinearColors.Get();
+	TargetColorChannels = InArgs._TargetColorChannels.Get();
+	bUseAlpha = InArgs._UseAlpha;
+	bOnlyRefreshOnMouseUp = InArgs._OnlyRefreshOnMouseUp.Get();
+	bOnlyRefreshOnOk = InArgs._OnlyRefreshOnOk.Get();
+	OnColorCommitted = InArgs._OnColorCommitted;
+	PreColorCommitted = InArgs._PreColorCommitted;
+	OnColorPickerCancelled = InArgs._OnColorPickerCancelled;
+	OnInteractivePickBegin = InArgs._OnInteractivePickBegin;
+	OnInteractivePickEnd = InArgs._OnInteractivePickEnd;
+	OnColorPickerWindowClosed = InArgs._OnColorPickerWindowClosed;
+	ParentWindowPtr = InArgs._ParentWindow.Get();
+	DisplayGamma = InArgs._DisplayGamma;
+	bClosedViaOkOrCancel = false;
+
+	// We need a parent window to set the close callback
+	if (ParentWindowPtr.IsValid())
+	{
+		ParentWindowPtr.Pin()->SetOnWindowClosed(FOnWindowClosed::CreateSP(this, &SColorPicker::HandleParentWindowClosed));
+	}
+
+	bColorPickerIsInlineVersion = InArgs._DisplayInlineVersion;
+	bIsInteractive = false;
+	bPerfIsTooSlowToUpdate = false;
+
+	BackupColors();
+
+	BeginAnimation(FLinearColor(ForceInit), CurrentColorHSV);
+
+	bool bAdvancedSectionExpanded = false;
+
+	if (!FPaths::FileExists(GEditorUserSettingsIni))
+	{
+		bool WheelMode = true;
+
+		GConfig->GetBool(TEXT("ColorPickerUI"), TEXT("bWheelMode"), WheelMode, GEditorUserSettingsIni);
+		GConfig->GetBool(TEXT("ColorPickerUI"), TEXT("bAdvancedSectionExpanded"), bAdvancedSectionExpanded, GEditorUserSettingsIni);
+		GConfig->GetBool(TEXT("ColorPickerUI"), TEXT("bSRGBEnabled"), SColorThemesViewer::bSRGBEnabled, GEditorUserSettingsIni);
+
+		CurrentMode = WheelMode ? EColorPickerModes::Wheel : EColorPickerModes::Spectrum;
+	}
+
+	bAdvancedSectionExpanded |= InArgs._ExpandAdvancedSection;
+
+	if (bColorPickerIsInlineVersion)
+	{
+		GenerateInlineColorPickerContent();
+	}
+	else
+	{
+		GenerateDefaultColorPickerContent(bAdvancedSectionExpanded);
+	}
+}
+
+
+/* SColorPicker implementation
+ *****************************************************************************/
+
+void SColorPicker::BackupColors()
+{
+	OldTargetFColors.Empty();
+	for (int32 i = 0; i < TargetFColors.Num(); ++i)
+	{
+		OldTargetFColors.Add( *TargetFColors[i] );
+	}
+
+	OldTargetLinearColors.Empty();
+	for (int32 i = 0; i < TargetLinearColors.Num(); ++i)
+	{
+		OldTargetLinearColors.Add( *TargetLinearColors[i] );
+	}
+
+	OldTargetColorChannels.Empty();
+	for (int32 i = 0; i < TargetColorChannels.Num(); ++i)
+	{
+		// Remap the color channel as a linear color for ease
+		const FColorChannels& Channel = TargetColorChannels[i];
+		const FLinearColor Color( Channel.Red ? *Channel.Red : 0.f, Channel.Green ? *Channel.Green : 0.f, Channel.Blue ? *Channel.Blue : 0.f, Channel.Alpha ? *Channel.Alpha : 0.f );
+		OldTargetColorChannels.Add( Color );
+	}
+}
+
+
+void SColorPicker::RestoreColors()
+{
+	check(TargetFColors.Num() == OldTargetFColors.Num());
+
+	for (int32 i = 0; i < TargetFColors.Num(); ++i)
+	{
+		*TargetFColors[i] = OldTargetFColors[i];
+	}
+
+	check(TargetLinearColors.Num() == OldTargetLinearColors.Num());
+
+	for (int32 i = 0; i < TargetLinearColors.Num(); ++i)
+	{
+		*TargetLinearColors[i] = OldTargetLinearColors[i];
+	}
+
+	check(TargetColorChannels.Num() == OldTargetColorChannels.Num());
+
+	for (int32 i = 0; i < TargetColorChannels.Num(); ++i)
+	{
+		// Copy back out of the linear to the color channel
+		FColorChannels& Channel = TargetColorChannels[i];
+		const FLinearColor& OldChannel = OldTargetColorChannels[i];
+		if (Channel.Red)
+		{
+			*Channel.Red = OldChannel.R;
+		}
+		if (Channel.Green)
+		{
+			*Channel.Green = OldChannel.G;
+		}
+		if (Channel.Blue)
+		{
+			*Channel.Blue = OldChannel.B;
+		}
+		if (Channel.Alpha)
+		{
+			*Channel.Alpha = OldChannel.A;
+		}
+	}
+}
+
+
+void SColorPicker::SetColors(const FLinearColor& InColor)
+{
+	for (int32 i = 0; i < TargetFColors.Num(); ++i)
+	{
+		*TargetFColors[i] = InColor.ToFColor(false);
+	}
+
+	for (int32 i = 0; i < TargetLinearColors.Num(); ++i)
+	{
+		*TargetLinearColors[i] = InColor;
+	}
+
+	for (int32 i = 0; i < TargetColorChannels.Num(); ++i)
+	{
+		// Only set those channels who have a valid ptr
+		FColorChannels& Channel = TargetColorChannels[i];
+		if (Channel.Red)
+		{
+			*Channel.Red = InColor.R;
+		}
+		if (Channel.Green)
+		{
+			*Channel.Green = InColor.G;
+		}
+		if (Channel.Blue)
+		{
+			*Channel.Blue = InColor.B;
+		}
+		if (Channel.Alpha)
+		{
+			*Channel.Alpha = InColor.A;
+		}
+	}
+}
+
+
+BEGIN_SLATE_FUNCTION_BUILD_OPTIMIZATION
+void SColorPicker::GenerateDefaultColorPickerContent( bool bAdvancedSectionExpanded )
+{	
+	// The height of the gradient bars beneath the sliders
+	const FSlateFontInfo SmallLayoutFont = FCoreStyle::Get().GetFontStyle("ColorPicker.Font");
+
+	TSharedPtr<SColorThemesViewer> ThemesViewer = ColorThemesViewer.Pin();
+
+	if (!ThemesViewer.IsValid())
+	{
+		ThemesViewer = SNew(SColorThemesViewer);
+		ColorThemesViewer = ThemesViewer;
+	}
+
+	ThemesViewer->OnCurrentThemeChanged().AddSP(this, &SColorPicker::HandleThemesViewerThemeChanged);
+	ThemesViewer->SetUseAlpha(bUseAlpha);
+	ThemesViewer->MenuToStandardNoReturn();
+
+	// The standard button to open the color themes can temporarily become a trash for colors
+	ColorThemeComboButton = SNew(SComboButton)
+		.ContentPadding(3.0f)
+		.MenuPlacement(MenuPlacement_ComboBox)
+		.ToolTipText(LOCTEXT("OpenThemeManagerToolTip", "Open Color Theme Manager"));
+
+	ColorThemeComboButton->SetMenuContent(ThemesViewer.ToSharedRef());
+
+	SmallTrash = SNew(SHorizontalBox)
+	+ SHorizontalBox::Slot()
+		.AutoWidth()
+		[
+			SNew(SColorTrash)
+				.UsesSmallIcon(true)
+		];
+	
+	this->ChildSlot
+	[
+		SNew(SVerticalBox)
+
+		+ SVerticalBox::Slot()
+			.AutoHeight()
+			[
+				SNew(SGridPanel)
+					.FillColumn(0, 1.0f)
+
+				+ SGridPanel::Slot(0, 0)
+					.Padding(0.0f, 1.0f, 20.0f, 1.0f)
+					[
+						SNew(SHorizontalBox)
+
+						+ SHorizontalBox::Slot()
+							.FillWidth(1.0f)
+							.Padding(0.0f, 1.0f)
+							[
+								SNew(SOverlay)
+
+								+ SOverlay::Slot()
+									[
+										// color theme bar
+										SAssignNew(CurrentThemeBar, SThemeColorBlocksBar)
+											.ColorTheme(this, &SColorPicker::HandleThemeBarColorTheme)
+											.EmptyText(LOCTEXT("EmptyBarHint", "Drag & drop colors here to save"))
+											.HideTrashCallback(this, &SColorPicker::HideSmallTrash)
+											.ShowTrashCallback(this, &SColorPicker::ShowSmallTrash)
+											.ToolTipText(LOCTEXT("CurrentThemeBarToolTip", "Current Color Theme"))
+											.UseAlpha(SharedThis(this), &SColorPicker::HandleThemeBarUseAlpha)
+											.UseSRGB(SharedThis(this), &SColorPicker::HandleColorPickerUseSRGB)
+											.OnSelectColor(this, &SColorPicker::HandleThemeBarColorSelected)
+									]
+
+								// hack: need to fix SThemeColorBlocksBar::EmptyText to render properly
+								+ SOverlay::Slot()
+									.HAlign(HAlign_Center)
+									.VAlign(VAlign_Center)
+									[
+										SNew(STextBlock)
+											.Text(LOCTEXT("EmptyBarHint", "Drag & drop colors here to save"))
+											.Visibility(this, &SColorPicker::HandleThemeBarHintVisibility)
+									]
+							]
+
+						+ SHorizontalBox::Slot()
+							.AutoWidth()
+							[
+								// color theme selector
+								SAssignNew(ColorThemeButtonOrSmallTrash, SBorder)
+									.BorderImage(FStyleDefaults::GetNoBrush())
+									.Padding(0.0f)							
+							]
+					]
+
+				+ SGridPanel::Slot(1, 0)
+					.HAlign(HAlign_Right)
+					.VAlign(VAlign_Center)
+					[
+						// sRGB check box
+						SNew(SCheckBox)
+							.ToolTipText(LOCTEXT("SRGBCheckboxToolTip", "Toggle gamma corrected sRGB previewing"))
+							.IsChecked(this, &SColorPicker::HandleSRGBCheckBoxIsChecked)
+							.OnCheckStateChanged(this, &SColorPicker::HandleSRGBCheckBoxCheckStateChanged)
+							[
+								SNew(STextBlock)
+									.Text(LOCTEXT("SRGBCheckboxLabel", "sRGB"))
+							]
+					]
+
+				+ SGridPanel::Slot(0, 1)
+					.Padding(0.0f, 8.0f, 20.0f, 0.0f)
+					[
+						SNew(SBorder)
+							.BorderImage(FCoreStyle::Get().GetBrush("NoBorder"))
+							.Padding(0.0f)
+							.OnMouseButtonDown(this, &SColorPicker::HandleColorAreaMouseDown)
+							[
+								SNew(SOverlay)
+
+								// color wheel
+								+ SOverlay::Slot()
+									[
+										SNew(SHorizontalBox)
+
+										+ SHorizontalBox::Slot()
+											.FillWidth(1.0f)
+											.HAlign(HAlign_Center)
+											[
+												SNew(SColorWheel)
+													.SelectedColor(this, &SColorPicker::GetCurrentColor)
+													.Visibility(this, &SColorPicker::HandleColorPickerModeVisibility, EColorPickerModes::Wheel)
+													.OnValueChanged(this, &SColorPicker::HandleColorSpectrumValueChanged)
+													.OnMouseCaptureBegin(this, &SColorPicker::HandleInteractiveChangeBegin)
+													.OnMouseCaptureEnd(this, &SColorPicker::HandleInteractiveChangeEnd)
+											]
+
+										+ SHorizontalBox::Slot()
+											.AutoWidth()
+											.Padding(4.0f, 0.0f)
+											[
+												// saturation slider
+												MakeColorSlider(EColorPickerChannels::Saturation)
+											]
+
+										+ SHorizontalBox::Slot()
+											.AutoWidth()
+											[
+												// value slider
+												MakeColorSlider(EColorPickerChannels::Value)
+											]
+									]
+
+								// color spectrum
+								+ SOverlay::Slot()
+									[
+										SNew(SBox)
+											.HeightOverride(200.0f)
+											.WidthOverride(292.0f)
+											[
+												SNew(SColorSpectrum)
+													.SelectedColor(this, &SColorPicker::GetCurrentColor)
+													.Visibility(this, &SColorPicker::HandleColorPickerModeVisibility, EColorPickerModes::Spectrum)
+													.OnValueChanged(this, &SColorPicker::HandleColorSpectrumValueChanged)
+													.OnMouseCaptureBegin(this, &SColorPicker::HandleInteractiveChangeBegin)
+													.OnMouseCaptureEnd(this, &SColorPicker::HandleInteractiveChangeEnd)
+											]
+									]
+							]
+					]
+
+				+ SGridPanel::Slot(1, 1)
+					.Padding(0.0f, 8.0f, 0.0f, 0.0f)
+					[
+						SNew(SVerticalBox)
+
+						+ SVerticalBox::Slot()
+							.AutoHeight()
+							[
+								SNew(SBox)
+									.HeightOverride(100.0f)
+									.WidthOverride(70.0f)
+									[
+										// color preview
+										MakeColorPreviewBox()
+									]
+							]
+
+						+ SVerticalBox::Slot()
+							.AutoHeight()
+							.Padding(0.0f, 16.0f, 0.0f, 0.0f)
+							.VAlign(VAlign_Top)
+							[
+								SNew(SHorizontalBox)
+
+								+ SHorizontalBox::Slot()
+									.HAlign(HAlign_Left)
+									[
+										// mode selector
+										SNew(SButton)
+											.OnClicked(this, &SColorPicker::HandleColorPickerModeButtonClicked)
+											.Content()
+											[
+												SNew(SImage)
+													.Image(FCoreStyle::Get().GetBrush("ColorPicker.Mode"))
+													.ToolTipText(LOCTEXT("ColorPickerModeEToolTip", "Toggle between color wheel and color spectrum."))
+											]									
+									]
+
+								+ SHorizontalBox::Slot()
+									.HAlign(HAlign_Right)
+									[
+										// eye dropper
+										SNew(SEyeDropperButton)
+											.OnValueChanged(this, &SColorPicker::HandleRGBColorChanged)
+											.OnBegin(this, &SColorPicker::HandleInteractiveChangeBegin)
+											.OnComplete(this, &SColorPicker::HandleEyeDropperButtonComplete)
+											.DisplayGamma(DisplayGamma)								
+									]
+							]
+					]
+			]
+
+		// advanced settings
+		+ SVerticalBox::Slot()
+			.AutoHeight()
+			.Padding(0.0f, 4.0f, 0.0f, 0.0f)
+			[
+				SNew(SExpandableArea)
+					.AreaTitle(LOCTEXT("AdvancedAreaTitle", "Advanced"))
+					.BorderBackgroundColor(FLinearColor::Transparent)
+					.InitiallyCollapsed(!bAdvancedSectionExpanded)
+					.OnAreaExpansionChanged(this, &SColorPicker::HandleAdvancedAreaExpansionChanged)
+					.Padding(FMargin(0.0f, 1.0f, 0.0f, 8.0f))
+					.BodyContent()
+					[
+						SNew(SHorizontalBox)
+
+						// RGBA inputs
+						+ SHorizontalBox::Slot()
+							.Padding(0.0f, 0.0f, 4.0f, 0.0f)
+							[
+								SNew(SVerticalBox)
+
+								// Red
+								+ SVerticalBox::Slot()
+									[
+										MakeColorSpinBox(EColorPickerChannels::Red)
+									]
+
+								// Green
+								+ SVerticalBox::Slot()
+									.Padding(0.0f, 8.0f, 0.0f, 0.0f)
+									[
+										MakeColorSpinBox(EColorPickerChannels::Green)
+									]
+
+								// Blue
+								+ SVerticalBox::Slot()
+									.Padding(0.0f, 8.0f, 0.0f, 0.0f)
+									[
+										MakeColorSpinBox(EColorPickerChannels::Blue)
+									]
+
+								// Alpha
+								+ SVerticalBox::Slot()
+									.Padding(0.0f, 8.0f, 0.0f, 0.0f)
+									[
+										MakeColorSpinBox(EColorPickerChannels::Alpha)
+									]
+							]
+
+						// HSV & Hex inputs
+						+ SHorizontalBox::Slot()
+							.Padding(4.0f, 0.0f, 0.0f, 0.0f)
+							[
+								SNew(SVerticalBox)
+							
+								// Hue
+								+ SVerticalBox::Slot()
+									[
+										MakeColorSpinBox(EColorPickerChannels::Hue)
+									]
+
+								// Saturation
+								+ SVerticalBox::Slot()
+									.Padding(0.0f, 8.0f, 0.0f, 0.0f)
+									[
+										MakeColorSpinBox(EColorPickerChannels::Saturation)
+									]
+
+								// Value
+								+ SVerticalBox::Slot()
+									.Padding(0.0f, 8.0f, 0.0f, 0.0f)
+									[
+										MakeColorSpinBox(EColorPickerChannels::Value)
+									]
+
+								// Hex
+								+ SVerticalBox::Slot()
+									.HAlign(HAlign_Right)
+									.VAlign(VAlign_Top)
+									.Padding(0.0f, 8.0f, 0.0f, 0.0f)
+									[
+										SNew(SHorizontalBox)
+											.ToolTipText(LOCTEXT("HexSliderToolTip", "Hexadecimal Value"))
+
+										+ SHorizontalBox::Slot()
+											.AutoWidth()
+											.Padding(0.0f, 0.0f, 4.0f, 0.0f)
+											.VAlign(VAlign_Center)
+											[
+												SNew(STextBlock)
+													.Text(LOCTEXT("HexInputLabel", "Hex"))
+											]
+
+										+ SHorizontalBox::Slot()
+											.AutoWidth()
+											[
+												SNew(SEditableTextBox)
+													.MinDesiredWidth(72.0f)
+													.Text(this, &SColorPicker::HandleHexBoxText)
+													.OnTextCommitted(this, &SColorPicker::HandleHexInputTextCommitted)
+											]
+									]
+							]
+					]
+			]
+
+		// dialog buttons
+		+ SVerticalBox::Slot()
+			.AutoHeight()
+			.HAlign(HAlign_Right)
+			.VAlign(VAlign_Center)
+			[
+				SNew(SUniformGridPanel)
+					.MinDesiredSlotHeight(FCoreStyle::Get().GetFloat("StandardDialog.MinDesiredSlotHeight"))
+					.MinDesiredSlotWidth(FCoreStyle::Get().GetFloat("StandardDialog.MinDesiredSlotWidth"))
+					.SlotPadding(FCoreStyle::Get().GetMargin("StandardDialog.SlotPadding"))
+					.Visibility(ParentWindowPtr.IsValid() ? EVisibility::Visible : EVisibility::Collapsed)
+
+				+ SUniformGridPanel::Slot(0, 0)
+					[
+						// ok button
+						SNew(SButton)
+							.ContentPadding( FCoreStyle::Get().GetMargin("StandardDialog.ContentPadding") )
+							.HAlign(HAlign_Center)
+							.Text(LOCTEXT("OKButton", "OK"))
+							.OnClicked(this, &SColorPicker::HandleOkButtonClicked)
+					]
+
+				+ SUniformGridPanel::Slot(1, 0)
+					[
+						// cancel button
+						SNew(SButton)
+							.ContentPadding( FCoreStyle::Get().GetMargin("StandardDialog.ContentPadding") )
+							.HAlign(HAlign_Center)
+							.Text(LOCTEXT("CancelButton", "Cancel"))
+							.OnClicked(this, &SColorPicker::HandleCancelButtonClicked)
+					]
+			]
+	];
+	
+	HideSmallTrash();
+}
+END_SLATE_FUNCTION_BUILD_OPTIMIZATION
+
+
+void SColorPicker::Tick( const FGeometry& AllottedGeometry, const double InCurrentTime, const float InDeltaTime )
+{
+	// animate the wheel and all the sliders
+	static const float AnimationTime = 0.2f;
+	if (CurrentTime < AnimationTime)
+	{
+		CurrentColorHSV = FMath::Lerp(ColorBegin, ColorEnd, CurrentTime / AnimationTime);
+		if (CurrentColorHSV.R < 0.f)
+		{
+			CurrentColorHSV.R += 360.f;
+		}
+		else if (CurrentColorHSV.R > 360.f)
+		{
+			CurrentColorHSV.R -= 360.f;
+		}
+
+		CurrentTime += InDeltaTime;
+		if (CurrentTime >= AnimationTime)
+		{
+			CurrentColorHSV = ColorEnd;
+		}
+
+		CurrentColorRGB = CurrentColorHSV.HSVToLinearRGB();
+	}
+}
+
+
+void SColorPicker::GenerateInlineColorPickerContent()
+{
+	ChildSlot
+	[
+		SNew(SHorizontalBox)
+
+		+ SHorizontalBox::Slot()
+			.HAlign(HAlign_Center)
+			.AutoWidth()
+			[
+				SNew(SVerticalBox)
+
+				+ SVerticalBox::Slot()
+					.AutoHeight()
+					[
+						SNew(SColorWheel)
+							.SelectedColor(this, &SColorPicker::GetCurrentColor)
+							.OnValueChanged(this, &SColorPicker::HandleHSVColorChanged)
+							.OnMouseCaptureBegin(this, &SColorPicker::HandleInteractiveChangeBegin)
+							.OnMouseCaptureEnd(this, &SColorPicker::HandleInteractiveChangeEnd)
+					]
+
+				+ SVerticalBox::Slot()
+					.Padding(FMargin(0.0f, 2.0f, 0.0f, 0.0f))
+					.AutoHeight()
+					[
+						SNew(SColorValueSlider)
+							.SelectedColor(this, &SColorPicker::GetCurrentColor)
+							.OnValueChanged(this, &SColorPicker::HandleHSVColorChanged)
+							.OnMouseCaptureBegin(this, &SColorPicker::HandleInteractiveChangeBegin)
+							.OnMouseCaptureEnd(this, &SColorPicker::HandleInteractiveChangeEnd)
+					]
+			]
+	];
+}
+
+
+
+
+
+void SColorPicker::DiscardColor()
+{
+	if (OnColorPickerCancelled.IsBound())
+	{
+		// let the user decide what to do about cancel
+		OnColorPickerCancelled.Execute( OldColor );
+	}
+	else
+	{	
+		SetNewTargetColorHSV(OldColor, true);
+		RestoreColors();
+	}
+}
+
+
+bool SColorPicker::SetNewTargetColorHSV( const FLinearColor& NewValue, bool bForceUpdate /*= false*/ )
+{
+	CurrentColorHSV = NewValue;
+	CurrentColorRGB = NewValue.HSVToLinearRGB();
+
+	return ApplyNewTargetColor(bForceUpdate);
+}
+
+
+bool SColorPicker::SetNewTargetColorRGB( const FLinearColor& NewValue, bool bForceUpdate /*= false*/ )
+{
+	CurrentColorRGB = NewValue;
+	CurrentColorHSV = NewValue.LinearRGBToHSV();
+
+	return ApplyNewTargetColor(bForceUpdate);
+}
+
+
+bool SColorPicker::ApplyNewTargetColor( bool bForceUpdate /*= false*/ )
+{
+	bool bUpdated = false;
+
+	if ((bForceUpdate || (!bOnlyRefreshOnMouseUp && !bPerfIsTooSlowToUpdate)) && (!bOnlyRefreshOnOk || bColorPickerIsInlineVersion))
+	{
+		double StartUpdateTime = FPlatformTime::Seconds();
+		UpdateColorPickMouseUp();
+		double EndUpdateTime = FPlatformTime::Seconds();
+
+		if (EndUpdateTime - StartUpdateTime > MAX_ALLOWED_UPDATE_TIME)
+		{
+			bPerfIsTooSlowToUpdate = true;
+		}
+
+		bUpdated = true;
+	}
+
+	return bUpdated;
+}
+
+
+void SColorPicker::UpdateColorPickMouseUp()
+{
+	if (!bOnlyRefreshOnOk || bColorPickerIsInlineVersion)
+	{
+		UpdateColorPick();
+	}
+}
+
+
+void SColorPicker::UpdateColorPick()
+{
+	bPerfIsTooSlowToUpdate = false;
+	FLinearColor OutColor = CurrentColorRGB;
+
+	PreColorCommitted.ExecuteIfBound(OutColor);
+
+	SetColors(OutColor);
+	OnColorCommitted.ExecuteIfBound(OutColor);
+	
+	// This callback is only necessary for wx backwards compatibility
+	FCoreDelegates::ColorPickerChanged.Broadcast();
+}
+
+
+void SColorPicker::BeginAnimation( FLinearColor Start, FLinearColor End )
+{
+	ColorEnd = End;
+	ColorBegin = Start;
+	CurrentTime = 0.f;
+	
+	// wraparound with hue
+	float HueDif = FMath::Abs(ColorBegin.R - ColorEnd.R);
+	if (FMath::Abs(ColorBegin.R + 360.f - ColorEnd.R) < HueDif)
+	{
+		ColorBegin.R += 360.f;
+	}
+	else if (FMath::Abs(ColorBegin.R - 360.f - ColorEnd.R) < HueDif)
+	{
+		ColorBegin.R -= 360.f;
+	}
+}
+
+
+void SColorPicker::HideSmallTrash()
+{
+	if (ColorThemeComboButton.IsValid())
+	{
+		ColorThemeButtonOrSmallTrash->SetContent(ColorThemeComboButton.ToSharedRef());
+	}
+	else
+	{
+		ColorThemeButtonOrSmallTrash->ClearContent();
+	}
+}
+
+
+void SColorPicker::ShowSmallTrash()
+{
+	if ( SmallTrash.IsValid() )
+	{
+		ColorThemeButtonOrSmallTrash->SetContent(SmallTrash.ToSharedRef());
+	}
+	else
+	{
+		ColorThemeButtonOrSmallTrash->ClearContent();
+	}
+}
+
+
+/* SColorPicker implementation
+ *****************************************************************************/
+
+void SColorPicker::CycleMode()
+{
+	if (CurrentMode == EColorPickerModes::Spectrum)
+	{
+		CurrentMode = EColorPickerModes::Wheel;
+	}
+	else
+	{
+		CurrentMode = EColorPickerModes::Spectrum;
+	}
+}
+
+
+BEGIN_SLATE_FUNCTION_BUILD_OPTIMIZATION
+TSharedRef<SWidget> SColorPicker::MakeColorSlider( EColorPickerChannels Channel ) const
+{
+	FText SliderTooltip;
+
+	switch (Channel)
+	{
+	case EColorPickerChannels::Red:			SliderTooltip = LOCTEXT("RedSliderToolTip", "Red"); break;
+	case EColorPickerChannels::Green:		SliderTooltip = LOCTEXT("GreenSliderToolTip", "Green"); break;
+	case EColorPickerChannels::Blue:		SliderTooltip = LOCTEXT("BlueSliderToolTip", "Blue"); break;
+	case EColorPickerChannels::Alpha:		SliderTooltip = LOCTEXT("AlphaSliderToolTip", "Alpha"); break;
+	case EColorPickerChannels::Hue:			SliderTooltip = LOCTEXT("HueSliderToolTip", "Hue"); break;
+	case EColorPickerChannels::Saturation:	SliderTooltip = LOCTEXT("SaturationSliderToolTip", "Saturation"); break;
+	case EColorPickerChannels::Value:		SliderTooltip = LOCTEXT("ValueSliderToolTip", "Value"); break;
+	default:
+		return SNullWidget::NullWidget;
+	}
+
+	return SNew(SOverlay)
+		.ToolTipText(SliderTooltip)
+
+	+ SOverlay::Slot()
+		.Padding(FMargin(4.0f, 0.0f))
+		[
+			SNew(SSimpleGradient)
+				.EndColor(this, &SColorPicker::HandleColorSliderEndColor, Channel)
+				.StartColor(this, &SColorPicker::HandleColorSliderStartColor, Channel)
+				.Orientation(Orient_Horizontal)
+				.UseSRGB(SharedThis(this), &SColorPicker::HandleColorPickerUseSRGB)
+		]
+
+	+ SOverlay::Slot()
+		[
+			SNew(SSlider)
+				.IndentHandle(false)
+				.Orientation(Orient_Vertical)
+				.SliderBarColor(FLinearColor::Transparent)
+				.Style(&FCoreStyle::Get().GetWidgetStyle<FSliderStyle>("ColorPicker.Slider"))
+				.Value(this, &SColorPicker::HandleColorSpinBoxValue, Channel)
+				.OnMouseCaptureBegin(this, &SColorPicker::HandleInteractiveChangeBegin)
+				.OnMouseCaptureEnd(this, &SColorPicker::HandleInteractiveChangeEnd)
+				.OnValueChanged(this, &SColorPicker::HandleColorSpinBoxValueChanged, Channel)
+		];
+}
+
+
+TSharedRef<SWidget> SColorPicker::MakeColorSpinBox( EColorPickerChannels Channel ) const
+{
+	if ((Channel == EColorPickerChannels::Alpha) && !bUseAlpha.Get())
+	{
+		return SNullWidget::NullWidget;
+	}
+
+	const int32 GradientHeight = 6;
+	const float HDRMaxValue = (TargetFColors.Num()) ? 1.f : FLT_MAX;
+	const FSlateFontInfo SmallLayoutFont = FCoreStyle::Get().GetFontStyle("ColorPicker.Font");
+
+	// create gradient widget
+	TSharedPtr<SWidget> GradientWidget;
+	
+	if (Channel == EColorPickerChannels::Hue)
+	{
+		TArray<FLinearColor> HueGradientColors;
+
+		for (int32 i = 0; i < 7; ++i)
+		{
+			HueGradientColors.Add( FLinearColor((i % 6) * 60.f, 1.f, 1.f).HSVToLinearRGB() );
+		}
+
+		GradientWidget = SNew(SComplexGradient)
+			.GradientColors(HueGradientColors);
+	}
+	else
+	{
+		GradientWidget = SNew(SSimpleGradient)
+			.StartColor(this, &SColorPicker::GetGradientStartColor, Channel)
+			.EndColor(this, &SColorPicker::GetGradientEndColor, Channel)
+			.HasAlphaBackground(Channel == EColorPickerChannels::Alpha)
+			.UseSRGB(SharedThis(this), &SColorPicker::HandleColorPickerUseSRGB);
+	}
+	
+	// create spin box
+	float MaxValue;
+	FText SliderLabel;
+	FText SliderTooltip;
+
+	switch (Channel)
+	{
+	case EColorPickerChannels::Red:
+		MaxValue = HDRMaxValue;
+		SliderLabel = LOCTEXT("RedSliderLabel", "R");
+		SliderTooltip = LOCTEXT("RedSliderToolTip", "Red");
+		break;
+		
+	case EColorPickerChannels::Green:
+		MaxValue = HDRMaxValue;
+		SliderLabel = LOCTEXT("GreenSliderLabel", "G");
+		SliderTooltip = LOCTEXT("GreenSliderToolTip", "Green");
+		break;
+		
+	case EColorPickerChannels::Blue:
+		MaxValue = HDRMaxValue;
+		SliderLabel = LOCTEXT("BlueSliderLabel", "B");
+		SliderTooltip = LOCTEXT("BlueSliderToolTip", "Blue");
+		break;
+		
+	case EColorPickerChannels::Alpha:
+		MaxValue = HDRMaxValue;
+		SliderLabel = LOCTEXT("AlphaSliderLabel", "A");
+		SliderTooltip = LOCTEXT("AlphaSliderToolTip", "Alpha");
+		break;
+		
+	case EColorPickerChannels::Hue:
+		MaxValue = 359.0f;
+		SliderLabel = LOCTEXT("HueSliderLabel", "H");
+		SliderTooltip = LOCTEXT("HueSliderToolTip", "Hue");
+		break;
+		
+	case EColorPickerChannels::Saturation:
+		MaxValue = 1.0f;
+		SliderLabel = LOCTEXT("SaturationSliderLabel", "S");
+		SliderTooltip = LOCTEXT("SaturationSliderToolTip", "Saturation");
+		break;
+		
+	case EColorPickerChannels::Value:
+		MaxValue = HDRMaxValue;
+		SliderLabel = LOCTEXT("ValueSliderLabel", "V");
+		SliderTooltip = LOCTEXT("ValueSliderToolTip", "Value");
+		break;
+
+	default:
+		return SNullWidget::NullWidget;
+	}
+	
+	return SNew(SHorizontalBox)
+
+	+ SHorizontalBox::Slot()
+		.AutoWidth()
+		.Padding(0.0f, 0.0f, 4.0f, 0.0f)
+		.VAlign(VAlign_Center)
+		[
+			SNew(STextBlock)
+				.Text(SliderLabel)
+		]
+
+	+ SHorizontalBox::Slot()
+		.FillWidth(1.0f)
+		[
+			SNew(SVerticalBox)
+				.ToolTipText(SliderTooltip)
+
+			+ SVerticalBox::Slot()
+				.AutoHeight()
+				[
+					SNew(SSpinBox<float>)
+						.MinValue(0.0f)
+						.MaxValue(MaxValue)
+						.MinSliderValue(0.0f)
+						.MaxSliderValue(Channel == EColorPickerChannels::Hue ? 359.0f : 1.0f)
+						.Delta(Channel == EColorPickerChannels::Hue ? 1.0f : 0.001f)
+						.Font(SmallLayoutFont)
+						.Value(this, &SColorPicker::HandleColorSpinBoxValue, Channel)
+						.OnBeginSliderMovement(this, &SColorPicker::HandleInteractiveChangeBegin)
+						.OnEndSliderMovement(this, &SColorPicker::HandleInteractiveChangeEnd)
+						.OnValueChanged(this, &SColorPicker::HandleColorSpinBoxValueChanged, Channel)
+				]
+
+			+ SVerticalBox::Slot()
+				.AutoHeight()
+				[
+					SNew(SBox)
+						.HeightOverride(GradientHeight)
+						[
+							GradientWidget.ToSharedRef()
+						]
+				]
+		];
+}
+
+
+TSharedRef<SWidget> SColorPicker::MakeColorPreviewBox() const
+{
+	return SNew(SOverlay)
+
+	+ SOverlay::Slot()
+		[
+			// preview blocks
+			SNew(SVerticalBox)
+
+			+ SVerticalBox::Slot()
+				.AutoHeight()
+				[
+					SNew(STextBlock)
+						.Text(LOCTEXT("OldColorLabel", "Old"))
+				]
+
+			+ SVerticalBox::Slot()
+				[
+					SNew(SHorizontalBox)
+
+					+ SHorizontalBox::Slot()
+						[
+							// old color
+							SNew(SColorBlock) 
+								.ColorIsHSV(true) 
+								.IgnoreAlpha(true)
+								.ToolTipText(LOCTEXT("OldColorToolTip", "Old color without alpha (drag to theme bar to save)"))
+								.Color(OldColor) 
+								.OnMouseButtonDown(this, &SColorPicker::HandleOldColorBlockMouseButtonDown, false)
+								.UseSRGB(SharedThis(this), &SColorPicker::HandleColorPickerUseSRGB)
+								.Cursor(EMouseCursor::GrabHand)
+						]
+
+					+ SHorizontalBox::Slot()
+						[
+							// old color (alpha)
+							SNew(SColorBlock) 
+								.ColorIsHSV(true) 
+								.ShowBackgroundForAlpha(true)
+								.ToolTipText(LOCTEXT("OldColorAlphaToolTip", "Old color with alpha (drag to theme bar to save)"))
+								.Color(OldColor) 
+								.Visibility(SharedThis(this), &SColorPicker::HandleAlphaColorBlockVisibility)
+								.OnMouseButtonDown(this, &SColorPicker::HandleOldColorBlockMouseButtonDown, true)
+								.UseSRGB(SharedThis(this), &SColorPicker::HandleColorPickerUseSRGB)
+								.Cursor(EMouseCursor::GrabHand)
+						]
+				]
+
+			+ SVerticalBox::Slot()
+				[
+					SNew(SHorizontalBox)
+
+					+ SHorizontalBox::Slot()
+						[
+							// new color
+							SNew(SColorBlock) 
+								.ColorIsHSV(true) 
+								.IgnoreAlpha(true)
+								.ToolTipText(LOCTEXT("NewColorToolTip", "New color without alpha (drag to theme bar to save)"))
+								.Color(this, &SColorPicker::GetCurrentColor)
+								.OnMouseButtonDown(this, &SColorPicker::HandleNewColorBlockMouseButtonDown, false)
+								.UseSRGB(SharedThis(this), &SColorPicker::HandleColorPickerUseSRGB)
+								.Cursor(EMouseCursor::GrabHand)
+						]
+
+					+ SHorizontalBox::Slot()
+						[
+							// new color (alpha)
+							SNew(SColorBlock) 
+								.ColorIsHSV(true) 
+								.ShowBackgroundForAlpha(true)
+								.ToolTipText(LOCTEXT("NewColorAlphaToolTip", "New color with alpha (drag to theme bar to save)"))
+								.Color(this, &SColorPicker::GetCurrentColor)
+								.Visibility(SharedThis(this), &SColorPicker::HandleAlphaColorBlockVisibility)
+								.OnMouseButtonDown(this, &SColorPicker::HandleNewColorBlockMouseButtonDown, true)
+								.UseSRGB(SharedThis(this), &SColorPicker::HandleColorPickerUseSRGB)
+								.Cursor(EMouseCursor::GrabHand)
+						]
+				]
+
+			+ SVerticalBox::Slot()
+				.AutoHeight()
+				[
+					SNew(STextBlock)
+						.Text(LOCTEXT("NewColorLabel", "New"))
+				]
+		]
+
+	+ SOverlay::Slot()
+		.VAlign(VAlign_Center)
+		[
+			// block separators
+			SNew(SHorizontalBox)
+
+			+ SHorizontalBox::Slot()
+				.FillWidth(0.5f)
+				.HAlign(HAlign_Left)
+				[
+					SNew(SBox)
+						.HeightOverride(2.0f)
+						.WidthOverride(4.0f)
+						[
+							SNew(SBorder)
+								.BorderImage(FCoreStyle::Get().GetBrush("ColorPicker.Separator"))
+								.Padding(0.0f)
+						]
+				]
+
+			+ SHorizontalBox::Slot()
+				.FillWidth(0.5f)
+				.HAlign(HAlign_Right)
+				[
+					SNew(SBox)
+						.HeightOverride(2.0f)
+						.WidthOverride(4.0f)
+						[
+							SNew(SBorder)
+								.BorderImage(FCoreStyle::Get().GetBrush("ColorPicker.Separator"))
+								.Padding(0.0f)
+						]
+				]
+		];
+}
+END_SLATE_FUNCTION_BUILD_OPTIMIZATION
+
+
+/* SColorPicker callbacks
+ *****************************************************************************/
+
+FLinearColor SColorPicker::GetGradientEndColor( EColorPickerChannels Channel ) const
+{
+	switch (Channel)
+	{
+	case EColorPickerChannels::Red:			return FLinearColor(1.0f, CurrentColorRGB.G, CurrentColorRGB.B, 1.0f);
+	case EColorPickerChannels::Green:		return FLinearColor(CurrentColorRGB.R, 1.0f, CurrentColorRGB.B, 1.0f);
+	case EColorPickerChannels::Blue:		return FLinearColor(CurrentColorRGB.R, CurrentColorRGB.G, 1.0f, 1.0f);
+	case EColorPickerChannels::Alpha:		return FLinearColor(CurrentColorRGB.R, CurrentColorRGB.G, CurrentColorRGB.B, 1.0f);
+	case EColorPickerChannels::Saturation:	return FLinearColor(CurrentColorHSV.R, 1.0f, CurrentColorHSV.B, 1.0f).HSVToLinearRGB();
+	case EColorPickerChannels::Value:		return FLinearColor(CurrentColorHSV.R, CurrentColorHSV.G, 1.0f, 1.0f).HSVToLinearRGB();
+	default:								return FLinearColor();
+	}
+}
+
+
+FLinearColor SColorPicker::GetGradientStartColor( EColorPickerChannels Channel ) const
+{
+	switch (Channel)
+	{
+	case EColorPickerChannels::Red:			return FLinearColor(0.0f, CurrentColorRGB.G, CurrentColorRGB.B, 1.0f);
+	case EColorPickerChannels::Green:		return FLinearColor(CurrentColorRGB.R, 0.0f, CurrentColorRGB.B, 1.0f);
+	case EColorPickerChannels::Blue:		return FLinearColor(CurrentColorRGB.R, CurrentColorRGB.G, 0.0f, 1.0f);
+	case EColorPickerChannels::Alpha:		return FLinearColor(CurrentColorRGB.R, CurrentColorRGB.G, CurrentColorRGB.B, 0.0f);
+	case EColorPickerChannels::Saturation:	return FLinearColor(CurrentColorHSV.R, 0.0f, CurrentColorHSV.B, 1.0f).HSVToLinearRGB();
+	case EColorPickerChannels::Value:		return FLinearColor(CurrentColorHSV.R, CurrentColorHSV.G, 0.0f, 1.0f).HSVToLinearRGB();
+	default:								return FLinearColor();
+	}
+}
+
+
+void SColorPicker::HandleAdvancedAreaExpansionChanged( bool Expanded )
+{
+	if (FPaths::FileExists(GEditorUserSettingsIni))
+	{
+		GConfig->SetBool(TEXT("ColorPickerUI"), TEXT("bAdvancedSectionExpanded"), Expanded, GEditorUserSettingsIni);
+	}
+}
+
+
+EVisibility SColorPicker::HandleAlphaColorBlockVisibility() const
+{
+	return bUseAlpha.Get() ? EVisibility::Visible : EVisibility::Collapsed;
+}
+
+
+FReply SColorPicker::HandleCancelButtonClicked()
+{
+	bClosedViaOkOrCancel = true;
+
+	DiscardColor();
+	ParentWindowPtr.Pin()->RequestDestroyWindow();
+
+	return FReply::Handled();
+}
+
+
+EVisibility SColorPicker::HandleColorPickerModeVisibility( EColorPickerModes Mode ) const
+{
+	return (CurrentMode == Mode) ? EVisibility::Visible : EVisibility::Hidden;
+}
+
+
+FLinearColor SColorPicker::HandleColorSliderEndColor( EColorPickerChannels Channel ) const
+{
+	switch (Channel)
+	{
+	case EColorPickerChannels::Red:			return FLinearColor(0.0f, 0.0f, 0.0f, 1.0f);
+	case EColorPickerChannels::Green:		return FLinearColor(0.0f, 0.0f, 0.0f, 1.0f);
+	case EColorPickerChannels::Blue:		return FLinearColor(0.0f, 0.0f, 0.0f, 1.0f);
+	case EColorPickerChannels::Alpha:		return FLinearColor(CurrentColorRGB.R, CurrentColorRGB.G, CurrentColorRGB.B, 0.0f);
+	case EColorPickerChannels::Saturation:	return FLinearColor(CurrentColorHSV.R, 0.0f, 1.0f, 1.0f).HSVToLinearRGB();
+	case EColorPickerChannels::Value:		return FLinearColor(CurrentColorHSV.R, CurrentColorHSV.G, 0.0f, 1.0f).HSVToLinearRGB();
+	default:								return FLinearColor();
+	}
+}
+
+
+FLinearColor SColorPicker::HandleColorSliderStartColor( EColorPickerChannels Channel ) const
+{
+	switch (Channel)
+	{
+	case EColorPickerChannels::Red:			return FLinearColor(1.0f, 0.0f, 0.0f, 1.0f);
+	case EColorPickerChannels::Green:		return FLinearColor(0.0f, 1.0f, 0.0f, 1.0f);
+	case EColorPickerChannels::Blue:		return FLinearColor(0.0f, 0.0f, 1.0f, 1.0f);
+	case EColorPickerChannels::Alpha:		return FLinearColor(CurrentColorRGB.R, CurrentColorRGB.G, CurrentColorRGB.B, 1.0f);
+	case EColorPickerChannels::Saturation:	return FLinearColor(CurrentColorHSV.R, 1.0f, 1.0f, 1.0f).HSVToLinearRGB();
+	case EColorPickerChannels::Value:		return FLinearColor(CurrentColorHSV.R, CurrentColorHSV.G, 1.0f, 1.0f).HSVToLinearRGB();
+	default:								return FLinearColor();
+	}
+}
+
+
+void SColorPicker::HandleColorSpectrumValueChanged( FLinearColor NewValue )
+{
+	SetNewTargetColorHSV(NewValue);
+}
+
+
+float SColorPicker::HandleColorSpinBoxValue( EColorPickerChannels Channel ) const
+{
+	switch (Channel)
+	{
+	case EColorPickerChannels::Red:			return CurrentColorRGB.R;
+	case EColorPickerChannels::Green:		return CurrentColorRGB.G;
+	case EColorPickerChannels::Blue:		return CurrentColorRGB.B;
+	case EColorPickerChannels::Alpha:		return CurrentColorRGB.A;
+	case EColorPickerChannels::Hue:			return CurrentColorHSV.R;
+	case EColorPickerChannels::Saturation:	return CurrentColorHSV.G;
+	case EColorPickerChannels::Value:		return CurrentColorHSV.B;
+	default:								return 0.0f;
+	}
+}
+
+
+void SColorPicker::HandleColorSpinBoxValueChanged( float NewValue, EColorPickerChannels Channel )
+{
+	int32 ComponentIndex;
+	bool IsHSV = false;
+
+	switch (Channel)
+	{
+	case EColorPickerChannels::Red:			ComponentIndex = 0; break;
+	case EColorPickerChannels::Green:		ComponentIndex = 1; break;
+	case EColorPickerChannels::Blue:		ComponentIndex = 2; break;
+	case EColorPickerChannels::Alpha:		ComponentIndex = 3; break;
+	case EColorPickerChannels::Hue:			ComponentIndex = 0; IsHSV = true; break;
+	case EColorPickerChannels::Saturation:	ComponentIndex = 1; IsHSV = true; break;
+	case EColorPickerChannels::Value:		ComponentIndex = 2; IsHSV = true; break;
+	default:								
+		return;
+	}
+
+	FLinearColor& NewColor = IsHSV ? CurrentColorHSV : CurrentColorRGB;
+
+	if (FMath::IsNearlyEqual(NewValue, NewColor.Component(ComponentIndex), KINDA_SMALL_NUMBER))
+	{
+		return;
+	}
+
+	NewColor.Component(ComponentIndex) = NewValue;
+
+	if (IsHSV)
+	{
+		SetNewTargetColorHSV(NewColor, !bIsInteractive);
+	}
+	else
+	{
+		SetNewTargetColorRGB(NewColor, !bIsInteractive);
+	}
+}
+
+
+void SColorPicker::HandleEyeDropperButtonComplete()
+{
+	bIsInteractive = false;
+
+	if (bOnlyRefreshOnMouseUp || bPerfIsTooSlowToUpdate)
+	{
+		UpdateColorPick();
+	}
+
+	OnInteractivePickEnd.ExecuteIfBound();
+}
+
+
+FText SColorPicker::HandleHexBoxText() const
+{
+	return FText::FromString(CurrentColorRGB.ToFColor(false).ToHex());
+}
+
+
+void SColorPicker::HandleHexInputTextCommitted( const FText& Text, ETextCommit::Type CommitType )
+{
+	if ((CommitType == ETextCommit::OnEnter) || (CommitType == ETextCommit::OnUserMovedFocus))
+	{
+		FColor Color = FColor::FromHex(Text.ToString());
+		SetNewTargetColorRGB(FLinearColor(Color.R / 255.0f, Color.G / 255.0f, Color.B / 255.0f, Color.A / 255.0f), false);
+	}	
+}
+
+
+void SColorPicker::HandleHSVColorChanged( FLinearColor NewValue )
+{
+	SetNewTargetColorHSV(NewValue);
+}
+
+
+void SColorPicker::HandleInteractiveChangeBegin()
+{
+	if( bIsInteractive && OnInteractivePickEnd.IsBound() )
+	{
+		OnInteractivePickEnd.Execute();
+	}
+
+	OnInteractivePickBegin.ExecuteIfBound();
+	bIsInteractive = true;
+}
+
+
+void SColorPicker::HandleInteractiveChangeEnd()
+{
+	HandleInteractiveChangeEnd(0.0f);
+}
+
+
+void SColorPicker::HandleInteractiveChangeEnd( float NewValue )
+{
+	bIsInteractive = false;
+
+	UpdateColorPickMouseUp();
+	OnInteractivePickEnd.ExecuteIfBound();
+}
+
+
+FReply SColorPicker::HandleColorAreaMouseDown( const FGeometry& MyGeometry, const FPointerEvent& MouseEvent )
+{
+	if (MouseEvent.GetEffectingButton() == EKeys::RightMouseButton)
+	{
+		CycleMode();
+
+		return FReply::Handled();
+	}
+
+	return FReply::Unhandled();
+}
+
+
+FReply SColorPicker::HandleColorPickerModeButtonClicked()
+{
+	CycleMode();
+
+	if (FPaths::FileExists(GEditorUserSettingsIni))
+	{
+		GConfig->SetBool(TEXT("ColorPickerUI"), TEXT("bWheelMode"), CurrentMode == EColorPickerModes::Wheel, GEditorUserSettingsIni);
+	}
+
+	return FReply::Handled();
+}
+
+
+FReply SColorPicker::HandleNewColorBlockMouseButtonDown( const FGeometry& MyGeometry, const FPointerEvent& MouseEvent, bool bCheckAlpha )
+{
+	if (MouseEvent.GetEffectingButton() == EKeys::LeftMouseButton)
+	{
+		TSharedRef<FColorDragDrop> Operation = FColorDragDrop::New(
+			CurrentColorHSV,
+			SColorThemesViewer::bSRGBEnabled,
+			bCheckAlpha ? bUseAlpha.Get() : false,
+			FSimpleDelegate::CreateSP(this, &SColorPicker::ShowSmallTrash),
+			FSimpleDelegate::CreateSP(this, &SColorPicker::HideSmallTrash)
+		);
+
+		return FReply::Handled().BeginDragDrop(Operation);
+	}
+
+	return FReply::Unhandled();
+}
+
+
+FReply SColorPicker::HandleOkButtonClicked()
+{
+	bClosedViaOkOrCancel = true;
+
+	if (bOnlyRefreshOnOk)
+	{
+		UpdateColorPick();
+	}
+
+	ParentWindowPtr.Pin()->RequestDestroyWindow();
+
+	return FReply::Handled();
+}
+
+
+FReply SColorPicker::HandleOldColorBlockMouseButtonDown( const FGeometry& MyGeometry, const FPointerEvent& MouseEvent, bool bCheckAlpha )
+{
+	if (MouseEvent.GetEffectingButton() == EKeys::LeftMouseButton)
+	{
+		TSharedRef<FColorDragDrop> Operation = FColorDragDrop::New(
+			OldColor,
+			SColorThemesViewer::bSRGBEnabled,
+			bCheckAlpha ? bUseAlpha.Get() : false,
+			FSimpleDelegate::CreateSP(this, &SColorPicker::ShowSmallTrash),
+			FSimpleDelegate::CreateSP(this, &SColorPicker::HideSmallTrash)
+		);
+
+		return FReply::Handled().BeginDragDrop(Operation);
+	}
+
+	return FReply::Unhandled();
+}
+
+
+bool SColorPicker::HandleColorPickerUseSRGB() const
+{
+	return SColorThemesViewer::bSRGBEnabled;
+}
+
+
+void SColorPicker::HandleParentWindowClosed( const TSharedRef<SWindow>& Window )
+{
+	check(Window == ParentWindowPtr.Pin());
+
+	// End picking interaction if still active
+	if( bIsInteractive && OnInteractivePickEnd.IsBound() )
+	{
+		OnInteractivePickEnd.Execute();
+		bIsInteractive = false;
+	}
+
+	// We always have to call the close callback
+	if (OnColorPickerWindowClosed.IsBound())
+	{
+		OnColorPickerWindowClosed.Execute(Window);
+	}
+
+	// If we weren't closed via the OK or Cancel button, we need to perform the default close action
+	if (!bClosedViaOkOrCancel && bOnlyRefreshOnOk)
+	{
+		DiscardColor();
+	}
+}
+
+
+void SColorPicker::HandleRGBColorChanged( FLinearColor NewValue )
+{
+	SetNewTargetColorRGB(NewValue);
+}
+
+
+void SColorPicker::HandleSRGBCheckBoxCheckStateChanged( ESlateCheckBoxState::Type InIsChecked )
+{
+	SColorThemesViewer::bSRGBEnabled = (InIsChecked == ESlateCheckBoxState::Checked);
+
+	if (FPaths::FileExists(GEditorUserSettingsIni))
+	{
+		GConfig->SetBool(TEXT("ColorPickerUI"), TEXT("bSRGBEnabled"), SColorThemesViewer::bSRGBEnabled, GEditorUserSettingsIni);
+	}
+}
+
+
+ESlateCheckBoxState::Type SColorPicker::HandleSRGBCheckBoxIsChecked() const
+{
+	return SColorThemesViewer::bSRGBEnabled ? ESlateCheckBoxState::Checked : ESlateCheckBoxState::Unchecked;
+}
+
+
+void SColorPicker::HandleThemeBarColorSelected( FLinearColor NewValue )
+{
+	// Force the alpha component to 1 when we don't care about the alpha
+	if (!bUseAlpha.Get())
+	{
+		NewValue.A = 1.0f;
+	}
+
+	BeginAnimation(CurrentColorHSV, NewValue);
+	SetNewTargetColorHSV(NewValue, true);
+}
+
+
+TSharedPtr<FColorTheme> SColorPicker::HandleThemeBarColorTheme() const
+{
+	TSharedPtr<SColorThemesViewer> ThemesViewer = ColorThemesViewer.Pin();
+
+	if (ThemesViewer.IsValid())
+	{
+		return ThemesViewer->GetCurrentColorTheme();
+	}
+
+	return nullptr;
+}
+
+
+EVisibility SColorPicker::HandleThemeBarHintVisibility() const
+{
+	TSharedPtr<SColorThemesViewer> ThemesViewer = ColorThemesViewer.Pin();
+
+	if (ThemesViewer.IsValid())
+	{
+		TSharedPtr<FColorTheme> SelectedTheme = ThemesViewer->GetCurrentColorTheme();
+
+		if (SelectedTheme.IsValid() && SelectedTheme->GetColors().Num() == 0)
+		{
+			return EVisibility::HitTestInvisible;
+		}
+	}
+
+	return EVisibility::Hidden;
+}
+
+
+bool SColorPicker::HandleThemeBarUseAlpha() const
+{
+	return bUseAlpha.Get();
+}
+
+
+void SColorPicker::HandleThemesViewerThemeChanged()
+{
+	if (CurrentThemeBar.IsValid())
+	{
+		CurrentThemeBar->RemoveRefreshCallback();
+		CurrentThemeBar->AddRefreshCallback();
+		CurrentThemeBar->Refresh();
+	}
+}
+
+
+/* Global functions
+ *****************************************************************************/
+
+/** A static color picker that everything should use. */
+static TWeakPtr<SWindow> ColorPickerWindow;
+
+
+bool OpenColorPicker(const FColorPickerArgs& Args)
+{
+	DestroyColorPicker();
+
+	bool Result = false;
+
+	// Consoles do not support opening new windows
+#if PLATFORM_DESKTOP
+	FLinearColor OldColor = Args.InitialColorOverride;
+
+	if (Args.ColorArray && Args.ColorArray->Num() > 0)
+	{
+		OldColor = (*Args.ColorArray)[0]->ReinterpretAsLinear();
+	}
+	else if (Args.LinearColorArray && Args.LinearColorArray->Num() > 0)
+	{
+		OldColor = *(*Args.LinearColorArray)[0];
+	}
+	else if (Args.ColorChannelsArray && Args.ColorChannelsArray->Num() > 0)
+	{
+		OldColor.R = (*Args.ColorChannelsArray)[0].Red ? *(*Args.ColorChannelsArray)[0].Red : 0.0f;
+		OldColor.G = (*Args.ColorChannelsArray)[0].Green ? *(*Args.ColorChannelsArray)[0].Green : 0.0f;
+		OldColor.B = (*Args.ColorChannelsArray)[0].Blue ? *(*Args.ColorChannelsArray)[0].Blue : 0.0f;
+		OldColor.A = (*Args.ColorChannelsArray)[0].Alpha ? *(*Args.ColorChannelsArray)[0].Alpha : 0.0f;
+	}
+	else
+	{
+		check(Args.OnColorCommitted.IsBound());
+	}
+		
+	// Determine the position of the window so that it will spawn near the mouse, but not go off the screen.
+	const FVector2D CursorPos = FSlateApplication::Get().GetCursorPos();
+	FSlateRect Anchor(CursorPos.X, CursorPos.Y, CursorPos.X, CursorPos.Y);
+
+	FVector2D AdjustedSummonLocation = FSlateApplication::Get().CalculatePopupWindowPosition( Anchor, SColorPicker::DEFAULT_WINDOW_SIZE, Orient_Horizontal );
+
+	TSharedPtr<SWindow> Window = SNew(SWindow)
+		.AutoCenter(EAutoCenter::None)
+		.ScreenPosition(AdjustedSummonLocation)
+		.SupportsMaximize(false)
+		.SupportsMinimize(false)
+		.SizingRule(ESizingRule::Autosized)
+		.Title(LOCTEXT("WindowHeader", "Color Picker"));
+
+	TSharedRef<SColorPicker> ColorPicker = SNew(SColorPicker)
+		.TargetColorAttribute(OldColor)
+		.TargetFColors(Args.ColorArray ? *Args.ColorArray : TArray<FColor*>())
+		.TargetLinearColors(Args.LinearColorArray ? *Args.LinearColorArray : TArray<FLinearColor*>())
+		.TargetColorChannels(Args.ColorChannelsArray ? *Args.ColorChannelsArray : TArray<FColorChannels>())
+		.UseAlpha(Args.bUseAlpha)
+		.ExpandAdvancedSection(Args.bExpandAdvancedSection)
+		.OnlyRefreshOnMouseUp(Args.bOnlyRefreshOnMouseUp && !Args.bIsModal)
+		.OnlyRefreshOnOk(Args.bOnlyRefreshOnOk || Args.bIsModal)
+		.OnColorCommitted(Args.OnColorCommitted)
+		.PreColorCommitted(Args.PreColorCommitted)
+		.OnColorPickerCancelled(Args.OnColorPickerCancelled)
+		.OnInteractivePickBegin(Args.OnInteractivePickBegin)
+		.OnInteractivePickEnd(Args.OnInteractivePickEnd)
+		.OnColorPickerWindowClosed(Args.OnColorPickerWindowClosed)
+		.ParentWindow(Window)
+		.DisplayGamma(Args.DisplayGamma);
+		
+	Window->SetContent(
+		SNew(SBox)
+			[
+				SNew(SBorder)
+					.BorderImage(FCoreStyle::Get().GetBrush("ToolPanel.GroupBorder"))
+					.Padding(FMargin(8.0f, 8.0f))
+					[
+						ColorPicker
+					]
+			]
+	);
+
+	if (Args.bIsModal)
+	{
+		FSlateApplication::Get().AddModalWindow(Window.ToSharedRef(), Args.ParentWidget);
+	}
+	else
+	{
+		if ( Args.ParentWidget.IsValid() )
+		{
+			// Find the window of the parent widget
+			FWidgetPath WidgetPath;
+			FSlateApplication::Get().GeneratePathToWidgetChecked( Args.ParentWidget.ToSharedRef(), WidgetPath );
+			Window = FSlateApplication::Get().AddWindowAsNativeChild( Window.ToSharedRef(), WidgetPath.GetWindow() );
+		}
+		else
+		{
+			Window = FSlateApplication::Get().AddWindow(Window.ToSharedRef());
+		}
+			
+	}
+
+	Result = true;
+
+	//hold on to the window created for external use...
+	ColorPickerWindow = Window;
+	
+#endif
+
+	return Result;
+}
+
+
+/**
+ * Destroys the current color picker. Necessary if the values the color picker
+ * currently targets become invalid.
+ */
+void DestroyColorPicker()
+{
+	if (ColorPickerWindow.IsValid())
+	{
+		ColorPickerWindow.Pin()->RequestDestroyWindow();
+		ColorPickerWindow.Reset();
+	}
+}
+
+
+#undef LOCTEXT_NAMESPACE
diff --git a/Engine/Source/Runtime/AppFramework/Private/Widgets/Colors/SColorThemes.cpp b/Engine/Source/Runtime/AppFramework/Private/Widgets/Colors/SColorThemes.cpp
new file mode 100644
--- /dev/null
+++ b/Engine/Source/Runtime/AppFramework/Private/Widgets/Colors/SColorThemes.cpp
@@ -0,0 +1,1164 @@
+// Copyright 1998-2014 Epic Games, Inc. All Rights Reserved.
+
+#include "AppFrameworkPrivatePCH.h"
+#include "SColorThemes.h"
+
+
+void FColorDragDrop::OnDrop( bool bDropWasHandled, const FPointerEvent& MouseEvent )
+{
+	HideTrash.ExecuteIfBound();
+
+	if (OriginBar.Pin().IsValid() && !bSetForDeletion)
+	{
+		OriginBar.Pin()->AddNewColorBlock(Color, OriginBarPosition);
+	}
+	
+	FDragDropOperation::OnDrop( bDropWasHandled, MouseEvent );
+}
+
+void FColorDragDrop::OnDragged( const class FDragDropEvent& DragDropEvent )
+{
+	if (CursorDecoratorWindow.IsValid())
+	{
+		CursorDecoratorWindow->MoveWindowTo( DragDropEvent.GetScreenSpacePosition() - BlockSize * 0.5f );
+	}
+}
+
+TSharedRef<FColorDragDrop> FColorDragDrop::New(FLinearColor InColor, bool bSRGB, bool bUseAlpha,
+	FSimpleDelegate TrashShowCallback, FSimpleDelegate TrashHideCallback,
+	TSharedPtr<SThemeColorBlocksBar> Origin, int32 OriginPosition)
+{
+	TSharedRef<FColorDragDrop> Operation = MakeShareable(new FColorDragDrop);
+
+	Operation->Color = InColor;
+	Operation->bUseSRGB = bSRGB;
+	Operation->bUseAlpha = bUseAlpha;
+	Operation->OriginBar = Origin;
+	Operation->OriginBarPosition = OriginPosition;
+	Operation->ShowTrash = TrashShowCallback;
+	Operation->HideTrash = TrashHideCallback;
+	Operation->bSetForDeletion = false;
+	Operation->BlockSize = FVector2D(32, 32);
+
+	Operation->ShowTrash.ExecuteIfBound();
+
+	Operation->Construct();
+
+	return Operation;
+}
+
+TSharedPtr<SWidget> FColorDragDrop::GetDefaultDecorator() const
+{
+	const bool bIgnoreAlpha = !bUseAlpha;
+	const bool bShowBackgroundForAlpha = bUseAlpha;
+
+	return SNew(SBorder) .Cursor(EMouseCursor::GrabHandClosed)
+		[
+			SNew(SColorBlock) 
+				.Color(Color)
+				.ColorIsHSV(true) 
+				.IgnoreAlpha(bIgnoreAlpha) 
+				.ShowBackgroundForAlpha(bShowBackgroundForAlpha) 
+				.UseSRGB(bUseSRGB)
+		];
+}
+
+
+FColorTheme::FColorTheme( const FString& InName, const TArray< TSharedPtr<FLinearColor> >& InColors )
+	: Name(InName)
+	, Colors(InColors)
+	, RefreshEvent()
+{ }
+
+
+void FColorTheme::InsertNewColor( TSharedPtr<FLinearColor> InColor, int32 InsertPosition )
+{
+	Colors.Insert(InColor, InsertPosition);
+	RefreshEvent.Broadcast();
+}
+
+
+int32 FColorTheme::FindApproxColor( const FLinearColor& InColor, float Tolerance ) const
+{
+	for (int32 ColorIndex = 0; ColorIndex < Colors.Num(); ++ColorIndex)
+	{
+		if (Colors[ColorIndex]->Equals(InColor, Tolerance))
+		{
+			return ColorIndex;
+		}
+	}
+
+	return INDEX_NONE;
+}
+
+
+void FColorTheme::RemoveAll()
+{
+	Colors.Empty();
+	RefreshEvent.Broadcast();
+}
+
+
+void FColorTheme::RemoveColor( int32 ColorIndex )
+{
+	Colors.RemoveAt(ColorIndex);
+	RefreshEvent.Broadcast();
+}
+
+
+int32 FColorTheme::RemoveColor( const TSharedPtr<FLinearColor> InColor )
+{
+	const int32 Position = Colors.Find(InColor);
+	if (Position != INDEX_NONE)
+	{
+		RemoveColor(Position);
+	}
+
+	return Position;
+}
+
+
+void SColorTrash::Construct( const FArguments& InArgs )
+{
+	bBorderActivated = false;
+
+	this->ChildSlot
+	[
+		SNew(SBorder)
+			.ToolTipText(NSLOCTEXT("ColorTrashWidget", "MouseOverToolTip", "Delete Color"))
+			.BorderImage(this, &SColorTrash::GetBorderStyle)
+		[
+			SNew(SHorizontalBox)
+			+SHorizontalBox::Slot() .HAlign(HAlign_Center) .FillWidth(1)
+			[
+				SNew(SImage) .Image(InArgs._UsesSmallIcon.Get() ?
+					FCoreStyle::Get().GetBrush("TrashCan_Small") :
+					FCoreStyle::Get().GetBrush("TrashCan"))
+			]
+		]
+	];
+}
+	
+/**
+* Called during drag and drop when the drag enters a widget.
+*
+* @param MyGeometry      The geometry of the widget receiving the event.
+* @param DragDropEvent   The drag and drop event.
+*
+* @return A reply that indicated whether this event was handled.
+*/
+void SColorTrash::OnDragEnter( const FGeometry& MyGeometry, const FDragDropEvent& DragDropEvent )
+{
+	if ( DragDropEvent.GetOperationAs<FColorDragDrop>().IsValid() )
+	{
+		bBorderActivated = true;
+	}
+}
+
+/**
+* Called during drag and drop when the drag leaves a widget.
+*
+* @param DragDropEvent   The drag and drop event.
+*
+* @return A reply that indicated whether this event was handled.
+*/
+void SColorTrash::OnDragLeave( const FDragDropEvent& DragDropEvent )
+{
+	if ( DragDropEvent.GetOperationAs<FColorDragDrop>().IsValid() )
+	{
+		bBorderActivated = false;
+	}
+}
+
+/**
+* Called when the user is dropping something onto a widget; terminates drag and drop.
+*
+* @param MyGeometry      The geometry of the widget receiving the event.
+* @param DragDropEvent   The drag and drop event.
+*
+* @return A reply that indicated whether this event was handled.
+*/
+FReply SColorTrash::OnDrop( const FGeometry& MyGeometry, const FDragDropEvent& DragDropEvent )
+{
+	TSharedPtr<FColorDragDrop> DragDropContent = DragDropEvent.GetOperationAs<FColorDragDrop>();
+	if ( DragDropContent.IsValid() )
+	{
+		DragDropContent->bSetForDeletion = true;
+
+		bBorderActivated = false;
+
+		return FReply::Handled();
+	}
+	else
+	{
+		return FReply::Unhandled();
+	}
+}
+
+const FSlateBrush* SColorTrash::GetBorderStyle() const
+{
+	return (bBorderActivated) ?
+		FCoreStyle::Get().GetBrush("FocusRectangle") :
+		FStyleDefaults::GetNoBrush();
+}
+
+
+
+/**
+* Construct the widget
+*
+* @param InArgs   Declaration from which to construct the widget.
+*/
+void SThemeColorBlock::Construct(const FArguments& InArgs )
+{
+	ColorPtr = InArgs._Color.Get();
+	OnSelectColor = InArgs._OnSelectColor;
+	ParentPtr = InArgs._Parent.Get();
+	ShowTrashCallback = InArgs._ShowTrashCallback;
+	HideTrashCallback = InArgs._HideTrashCallback;
+	bUseSRGB = InArgs._UseSRGB;
+	bUseAlpha = InArgs._UseAlpha;
+
+	DistanceDragged = 0;
+
+	const FSlateFontInfo SmallLayoutFont( FPaths::EngineContentDir() / TEXT("Slate/Fonts/Roboto-Regular.ttf"), 9 );
+
+	TSharedPtr<SToolTip> ColorTooltip =
+		SNew(SToolTip)
+		[
+			SNew(SHorizontalBox)
+			+SHorizontalBox::Slot().AutoWidth()
+			[
+				SNew(SVerticalBox)
+				+SVerticalBox::Slot().AutoHeight() .Padding(2)
+				[
+					SNew(SBox)
+					.WidthOverride(110)
+					.HeightOverride(110)
+					[
+						SNew(SColorBlock)
+							.Color(this, &SThemeColorBlock::GetColor)
+							.ColorIsHSV(true)
+							.IgnoreAlpha(TAttribute<bool>::Create(TAttribute<bool>::FGetter::CreateSP(this, &SThemeColorBlock::OnReadIgnoreAlpha)))
+							.ShowBackgroundForAlpha(TAttribute<bool>::Create(TAttribute<bool>::FGetter::CreateSP(this, &SThemeColorBlock::OnReadShowBackgroundForAlpha)))
+							.UseSRGB(bUseSRGB)
+					]
+				]
+				+SVerticalBox::Slot().AutoHeight() .Padding(2)
+				[
+					SNew(SHorizontalBox)
+					+SHorizontalBox::Slot() .FillWidth(1)
+					[
+						SNew(SVerticalBox)
+						+SVerticalBox::Slot().AutoHeight() .Padding(3)
+						[
+							SNew(STextBlock) .Font(SmallLayoutFont) .Text(this, &SThemeColorBlock::GetRedText)
+						]
+						+SVerticalBox::Slot().AutoHeight() .Padding(3)
+						[
+							SNew(STextBlock) .Font(SmallLayoutFont) .Text(this, &SThemeColorBlock::GetGreenText)
+						]
+						+SVerticalBox::Slot().AutoHeight() .Padding(3)
+						[
+							SNew(STextBlock) .Font(SmallLayoutFont) .Text(this, &SThemeColorBlock::GetBlueText)
+						]
+					]
+					+SHorizontalBox::Slot() .FillWidth(1)
+					[
+						SNew(SVerticalBox)
+						+SVerticalBox::Slot().AutoHeight() .Padding(3)
+						[
+							SNew(STextBlock) .Font(SmallLayoutFont) .Text(this, &SThemeColorBlock::GetHueText)
+						]
+						+SVerticalBox::Slot().AutoHeight() .Padding(3)
+						[
+							SNew(STextBlock) .Font(SmallLayoutFont) .Text(this, &SThemeColorBlock::GetSaturationText)
+						]
+						+SVerticalBox::Slot().AutoHeight() .Padding(3)
+						[
+							SNew(STextBlock) .Font(SmallLayoutFont) .Text(this, &SThemeColorBlock::GetValueText)
+						]
+					]
+				]
+				+SVerticalBox::Slot().AutoHeight() .Padding(2)
+				.HAlign(HAlign_Center)
+				[
+					SNew(STextBlock) .Font(SmallLayoutFont) .Text(this, &SThemeColorBlock::GetAlphaText) .Visibility(SharedThis(this), &SThemeColorBlock::OnGetAlphaVisibility)
+				]
+			]
+		];
+
+	this->ChildSlot
+	[
+		SNew(SBorder)
+			.BorderImage(this, &SThemeColorBlock::HandleBorderImage)
+			.BorderBackgroundColor(this, &SThemeColorBlock::HandleBorderColor)
+			.Padding(FMargin(1.0f))
+			.ToolTip(ColorTooltip)
+			[
+				SNew(SColorBlock)
+					.Color(this, &SThemeColorBlock::GetColor)
+					.ColorIsHSV(true)
+					.IgnoreAlpha(TAttribute<bool>::Create(TAttribute<bool>::FGetter::CreateSP(this, &SThemeColorBlock::OnReadIgnoreAlpha)))
+					.ShowBackgroundForAlpha(TAttribute<bool>::Create(TAttribute<bool>::FGetter::CreateSP(this, &SThemeColorBlock::OnReadShowBackgroundForAlpha)))
+					.UseSRGB(bUseSRGB)
+			]
+	];
+}
+
+FReply SThemeColorBlock::OnMouseButtonDown( const FGeometry& MyGeometry, const FPointerEvent& MouseEvent )
+{
+	if ( MouseEvent.GetEffectingButton() == EKeys::LeftMouseButton )
+	{
+		DistanceDragged = 0;
+		return FReply::Handled().DetectDrag( SharedThis(this), EKeys::LeftMouseButton ).CaptureMouse(SharedThis(this));
+	}
+	else
+	{
+		return FReply::Unhandled();
+	}
+}
+
+FReply SThemeColorBlock::OnMouseButtonUp( const FGeometry& MyGeometry, const FPointerEvent& MouseEvent )
+{
+	if ( MouseEvent.GetEffectingButton() == EKeys::LeftMouseButton && MyGeometry.IsUnderLocation(MouseEvent.GetScreenSpacePosition()) )
+	{
+		check(ColorPtr.IsValid());
+		OnSelectColor.ExecuteIfBound(*ColorPtr.Pin());
+		
+		return FReply::Handled().ReleaseMouseCapture();
+			
+	}
+	else
+	{
+		return FReply::Unhandled();
+	}
+}
+
+FReply SThemeColorBlock::OnDragDetected( const FGeometry& MyGeometry, const FPointerEvent& MouseEvent )
+{
+	if ( MouseEvent.IsMouseButtonDown( EKeys::LeftMouseButton ) && ParentPtr.IsValid() )
+	{
+		FLinearColor ColorToGrab = *ColorPtr.Pin();
+
+		int32 Position = ParentPtr.Pin()->RemoveColorBlock(ColorPtr.Pin());
+			
+		ParentPtr.Pin()->SetPlaceholderGrabOffset(MyGeometry.AbsoluteToLocal( MouseEvent.GetScreenSpacePosition() ));
+		
+		TSharedRef<FColorDragDrop> Operation = FColorDragDrop::New(ColorToGrab, bUseSRGB.Get(), bUseAlpha.Get(), 
+			ShowTrashCallback, HideTrashCallback, ParentPtr.Pin(), Position);
+		return FReply::Handled().BeginDragDrop(Operation);
+	}
+	
+	return FReply::Unhandled();
+}
+
+FLinearColor SThemeColorBlock::GetColor() const
+{
+	return ColorPtr.IsValid() ? *ColorPtr.Pin() : FLinearColor(ForceInit);
+}
+
+FSlateColor SThemeColorBlock::HandleBorderColor() const
+{
+	if (IsHovered())
+	{
+		return FSlateColor(FLinearColor::White);
+	}
+	
+	return FSlateColor(GetColor().HSVToLinearRGB().ToFColor(bUseSRGB.Get()));
+}
+
+const FSlateBrush* SThemeColorBlock::HandleBorderImage() const
+{
+	return IsHovered() ?
+		FCoreStyle::Get().GetBrush("FocusRectangle") :
+		FCoreStyle::Get().GetBrush("GenericWhiteBox");
+}
+
+#define LOCTEXT_NAMESPACE "ThemeColorBlock"
+FText SThemeColorBlock::GetRedText() const { return ColorPtr.IsValid() ? FormatToolTipText(LOCTEXT("Red", "R"), ColorPtr.Pin()->HSVToLinearRGB().R) : FText::GetEmpty(); }
+FText SThemeColorBlock::GetGreenText() const { return ColorPtr.IsValid() ? FormatToolTipText(LOCTEXT("Green", "G"), ColorPtr.Pin()->HSVToLinearRGB().G) : FText::GetEmpty(); }
+FText SThemeColorBlock::GetBlueText() const { return ColorPtr.IsValid() ? FormatToolTipText(LOCTEXT("Blue", "B"), ColorPtr.Pin()->HSVToLinearRGB().B) : FText::GetEmpty(); }
+FText SThemeColorBlock::GetAlphaText() const { return ColorPtr.IsValid() ? FormatToolTipText(LOCTEXT("Alpha", "A"), ColorPtr.Pin()->HSVToLinearRGB().A) : FText::GetEmpty(); }
+FText SThemeColorBlock::GetHueText() const { return ColorPtr.IsValid() ? FormatToolTipText(LOCTEXT("Hue", "H"), ColorPtr.Pin()->R) : FText::GetEmpty(); }
+FText SThemeColorBlock::GetSaturationText() const { return ColorPtr.IsValid() ? FormatToolTipText(LOCTEXT("Saturation", "S"), ColorPtr.Pin()->G) : FText::GetEmpty(); }
+FText SThemeColorBlock::GetValueText() const { return ColorPtr.IsValid() ? FormatToolTipText(LOCTEXT("Value", "V"), ColorPtr.Pin()->B) : FText::GetEmpty(); }
+
+FText SThemeColorBlock::FormatToolTipText(const FText& ColorIdentifier, float Value) const
+{
+	FFormatNamedArguments Args;
+	Args.Add(TEXT("Identifier"), ColorIdentifier);
+
+	if (Value >= 0.f)
+	{
+		static const float LogToLog10 = 1.f / FMath::Loge(10.f);
+		int32 PreRadixDigits = FMath::Max(0, int32(FMath::Loge(Value + KINDA_SMALL_NUMBER) * LogToLog10));
+
+		int32 Precision = FMath::Max(0, 2 - PreRadixDigits);
+
+		FNumberFormattingOptions FormatRules;
+		FormatRules.MinimumFractionalDigits = Precision;
+
+		Args.Add(TEXT("Value"), FText::AsNumber(Value, &FormatRules));
+	}
+	else
+	{
+		Args.Add(TEXT("Value"), FText::GetEmpty());
+	}
+
+	return FText::Format(LOCTEXT("ToolTipFormat", "{Identifier}: {Value}"), Args);
+}
+#undef LOCTEXT_NAMESPACE
+
+bool SThemeColorBlock::OnReadIgnoreAlpha() const
+{
+	return !bUseAlpha.Get();
+}
+
+bool SThemeColorBlock::OnReadShowBackgroundForAlpha() const
+{
+	return bUseAlpha.Get();
+}
+
+EVisibility SThemeColorBlock::OnGetAlphaVisibility() const
+{
+	return bUseAlpha.Get() ? EVisibility::Visible : EVisibility::Collapsed;
+}
+
+SThemeColorBlocksBar::SThemeColorBlocksBar()
+: ColorBlocks()
+{
+}
+
+
+/**
+* Panels arrange their children in a space described by the AllottedGeometry parameter. The results of the arrangement
+* should be returned by appending a FArrangedWidget pair for every child widget.
+*
+* @param AllottedGeometry    The geometry allotted for this widget by its parent.
+* @param ArrangedChildren    The array to which to add the WidgetGeometries that represent the arranged children.
+*/
+void SThemeColorBlocksBar::OnArrangeChildren( const FGeometry& AllottedGeometry, FArrangedChildren& ArrangedChildren ) const
+{
+	bool bPlaceHolderExists = NewColorPlaceholder.IsValid();
+
+	const int32 NumColorBlocks = bPlaceHolderExists ? ColorBlocks.Num() + 1 : ColorBlocks.Num();
+
+	if (NumColorBlocks > 0)
+	{
+		const float ChildPadding = 3.0f;
+		const FVector2D ChildSize = FVector2D(AllottedGeometry.Size.X / NumColorBlocks - ChildPadding, AllottedGeometry.Size.Y);
+
+		const float CurrentDragCenter = PlaceholderBlockOffset + ChildSize.X * 0.5f;
+
+		float XOffset = 0;
+		for (int32 i = 0; i < NumColorBlocks; ++i)
+		{
+			if (bPlaceHolderExists && XOffset <= CurrentDragCenter && CurrentDragCenter < (XOffset + ChildSize.X))
+			{
+				XOffset += ChildSize.X;
+				XOffset += ChildPadding;
+			}
+
+			if (i < ColorBlocks.Num())
+			{
+				ArrangedChildren.AddWidget(AllottedGeometry.MakeChild(ColorBlocks[i], FVector2D(XOffset, 0), ChildSize));
+				XOffset += ChildSize.X;
+				XOffset += ChildPadding;
+			}
+		}
+
+		if (bPlaceHolderExists && NewColorBlockPlaceholder.IsValid())
+		{
+			ArrangedChildren.AddWidget(AllottedGeometry.MakeChild(NewColorBlockPlaceholder.ToSharedRef(), FVector2D(PlaceholderBlockOffset, 0), ChildSize));
+		}
+	}
+	else if (EmptyHintTextBlock.IsValid())
+	{
+		ArrangedChildren.AddWidget(AllottedGeometry.MakeChild(EmptyHintTextBlock.ToSharedRef(), FVector2D::ZeroVector, AllottedGeometry.Size));
+	}
+}
+
+/**
+* A Panel's desired size in the space required to arrange of its children on the screen while respecting all of
+* the children's desired sizes and any layout-related options specified by the user. See StackPanel for an example.
+*
+* @return The desired size.
+*/
+FVector2D SThemeColorBlocksBar::ComputeDesiredSize() const
+{
+	return FVector2D(64, 16);
+}
+
+/**
+* All widgets must provide a way to access their children in a layout-agnostic way.
+* Panels store their children in Slots, which creates a dilemma. Most panels
+* can store their children in a TPanelChildren<Slot>, where the Slot class
+* provides layout information about the child it stores. In that case
+* GetChildren should simply return the TPanelChildren<Slot>.
+*/
+FChildren* SThemeColorBlocksBar::GetChildren()
+{
+	return &ColorBlocks;
+}
+
+void SThemeColorBlocksBar::OnDragEnter( const FGeometry& MyGeometry, const FDragDropEvent& DragDropEvent )
+{
+	TSharedPtr<FColorDragDrop> DragDropOperation = DragDropEvent.GetOperationAs<FColorDragDrop>();
+	if ( DragDropOperation.IsValid() )
+	{
+		NewColorPlaceholder = MakeShareable(new FLinearColor(DragDropOperation->Color));
+		NewColorBlockPlaceholder =
+			SNew(SThemeColorBlock) .Color(NewColorPlaceholder) .UseSRGB(bUseSRGB) .UseAlpha(bUseAlpha);
+	}
+}
+
+void SThemeColorBlocksBar::OnDragLeave( const FDragDropEvent& DragDropEvent )
+{
+	TSharedPtr<FColorDragDrop> DragDropOperation = DragDropEvent.GetOperationAs<FColorDragDrop>();
+	if ( DragDropOperation.IsValid() )
+	{
+		DragDropOperation->SetDecoratorVisibility(true);
+
+		DestroyPlaceholders();
+	}
+}
+
+/**
+* Called during drag and drop when the the mouse is being dragged over a widget.
+*
+* @param MyGeometry      The geometry of the widget receiving the event.
+* @param DragDropEvent   The drag and drop event.
+*
+* @return A reply that indicated whether this event was handled.
+*/
+FReply SThemeColorBlocksBar::OnDragOver( const FGeometry& MyGeometry, const FDragDropEvent& DragDropEvent )
+{
+	TSharedPtr<FColorDragDrop> DragDropOperation = DragDropEvent.GetOperationAs<FColorDragDrop>();
+	if ( DragDropOperation.IsValid() )
+	{
+		const float ChildSizeX = MyGeometry.Size.X / (ColorBlocks.Num() + 1);
+		const float GrabOffsetX = PlaceholderInitialGrabOffset.X;
+		PlaceholderBlockOffset = MyGeometry.AbsoluteToLocal(DragDropEvent.GetScreenSpacePosition()).X -
+			((GrabOffsetX == 0.f) ? ChildSizeX * 0.5f : GrabOffsetX);
+
+		DragDropOperation->SetDecoratorVisibility(false);
+			
+		return FReply::Handled();
+	}
+	else
+	{
+		return FReply::Unhandled();
+	}
+}
+
+/**
+* Called when the user is dropping something onto a widget; terminates drag and drop.
+*
+* @param MyGeometry      The geometry of the widget receiving the event.
+* @param DragDropEvent   The drag and drop event.
+*
+* @return A reply that indicated whether this event was handled.
+*/
+FReply SThemeColorBlocksBar::OnDrop( const FGeometry& MyGeometry, const FDragDropEvent& DragDropEvent )
+{
+	TSharedPtr<FColorDragDrop> DragDropOperation = DragDropEvent.GetOperationAs<FColorDragDrop>();
+	if ( DragDropOperation.IsValid() )
+	{
+		const float ChildSizeX = MyGeometry.Size.X / (ColorBlocks.Num() + 1);
+		const float CurrentDragCenter = PlaceholderBlockOffset + ChildSizeX * 0.5f;
+		const int32 LocID = FMath::Clamp(int32(CurrentDragCenter / ChildSizeX), 0, ColorBlocks.Num());
+
+		AddNewColorBlock(DragDropOperation->Color, LocID);
+
+		DragDropOperation->bSetForDeletion = true;
+
+		return FReply::Handled();
+	}
+	else
+	{
+		return FReply::Unhandled();
+	}
+}
+
+void SThemeColorBlocksBar::AddNewColorBlock(FLinearColor Color, int32 InsertPosition)
+{
+	// Force the alpha component to 1 when we don't care about the alpha, it'll allow the color to work with alpha-important widgets later
+	if( !bUseAlpha.Get() )
+	{
+		Color.A = 1.0f;
+	}
+	ColorTheme.Get()->InsertNewColor(MakeShareable(new FLinearColor(Color)), InsertPosition);
+
+	DestroyPlaceholders();
+
+	SColorThemesViewer::SaveColorThemesToIni();
+}
+
+int32 SThemeColorBlocksBar::RemoveColorBlock(TSharedPtr<FLinearColor> ColorToRemove)
+{
+	int32 Position = ColorTheme.Get()->RemoveColor(ColorToRemove);
+	
+	SColorThemesViewer::SaveColorThemesToIni();
+
+	return Position;
+}
+
+void SThemeColorBlocksBar::DestroyPlaceholders()
+{
+	NewColorBlockPlaceholder.Reset();
+	NewColorPlaceholder.Reset();
+	PlaceholderBlockOffset = 0.f;
+	PlaceholderInitialGrabOffset = FVector2D(ForceInit);
+}
+
+void SThemeColorBlocksBar::RemoveRefreshCallback()
+{
+	ColorTheme.Get()->OnRefresh().Remove(RefreshCallback);
+}
+
+void SThemeColorBlocksBar::AddRefreshCallback()
+{
+	ColorTheme.Get()->OnRefresh().Add(RefreshCallback);
+}
+
+void SThemeColorBlocksBar::Refresh()
+{
+	DestroyPlaceholders();
+	ColorBlocks.Empty();
+
+	check(ColorTheme.Get().IsValid());
+
+	const TArray< TSharedPtr<FLinearColor> >& Theme = ColorTheme.Get()->GetColors();
+	for (int32 i = 0; i < Theme.Num(); ++i)
+	{
+		ColorBlocks.Add(
+			SNew(SThemeColorBlock)
+				.Color(Theme[i])
+				.OnSelectColor(OnSelectColor)
+				.Parent(SharedThis(this))
+				.ShowTrashCallback(ShowTrashCallback)
+				.HideTrashCallback(HideTrashCallback)
+				.UseSRGB(bUseSRGB)
+				.UseAlpha(bUseAlpha)
+			);
+	}
+}
+
+void SThemeColorBlocksBar::SetPlaceholderGrabOffset(FVector2D GrabOffset)
+{
+	PlaceholderInitialGrabOffset = GrabOffset;
+}
+
+/**
+* Most panels do not create widgets as part of their implementation, so
+* they do not need to implement a Construct()
+*/
+void SThemeColorBlocksBar::Construct(const FArguments& InArgs)
+{
+	ColorTheme = InArgs._ColorTheme;
+	OnSelectColor = InArgs._OnSelectColor;
+	ShowTrashCallback = InArgs._ShowTrashCallback;
+	HideTrashCallback = InArgs._HideTrashCallback;
+	bUseSRGB = InArgs._UseSRGB;
+	bUseAlpha = InArgs._UseAlpha;
+	
+	RefreshCallback = FSimpleDelegate::CreateSP(this, &SThemeColorBlocksBar::Refresh);
+	AddRefreshCallback();
+
+	DestroyPlaceholders();
+
+	if (InArgs._EmptyText.Len() > 0)
+	{
+		EmptyHintTextBlock = SNew(SBorder)
+			.Padding(1.0f)
+			.Content()
+			[
+				SNew(STextBlock)
+					.Font(FSlateFontInfo(FPaths::EngineContentDir() / TEXT("Slate/Fonts/Roboto-Regular.ttf"), 8))
+					.Text(InArgs._EmptyText)
+			];
+	}
+
+	Refresh();
+}
+
+
+
+
+void SColorThemeBar::Construct(const FArguments& InArgs)
+{
+	ColorTheme = InArgs._ColorTheme.Get();
+	OnCurrentThemeChanged = InArgs._OnCurrentThemeChanged;
+	ShowTrashCallback = InArgs._ShowTrashCallback;
+	HideTrashCallback = InArgs._HideTrashCallback;
+	bUseSRGB = InArgs._UseSRGB;
+	bUseAlpha = InArgs._UseAlpha;
+
+	this->ChildSlot
+	[
+		SNew(SBox)
+			.Padding(5.0f)
+			[
+				SNew(SHorizontalBox)
+
+				+ SHorizontalBox::Slot()
+					.AutoWidth()
+					[
+						SAssignNew(ThemeNameText, STextBlock)
+							.Text(this, &SColorThemeBar::GetThemeName)
+							.Font(FSlateFontInfo(FPaths::EngineContentDir() / TEXT("Slate/Fonts/Roboto-Regular.ttf"), 10))
+					]
+
+				+ SHorizontalBox::Slot()
+					.FillWidth(1.0f)
+					.Padding(8.0f, 0.0f, 0.0f, 0.0f)
+					[
+						SNew(SThemeColorBlocksBar)
+							.ColorTheme(InArgs._ColorTheme)
+							.ShowTrashCallback(ShowTrashCallback)
+							.HideTrashCallback(HideTrashCallback)
+							.EmptyText(NSLOCTEXT("ColorThemesViewer", "NoColorsText", "No Colors Added Yet").ToString())
+							.UseSRGB(bUseSRGB)
+							.UseAlpha(bUseAlpha)
+					]
+			]
+	];
+}
+
+FString SColorThemeBar::GetThemeName() const
+{
+	return ColorTheme.Pin()->Name;
+}
+
+FReply SColorThemeBar::OnMouseButtonDown( const FGeometry& MyGeometry, const FPointerEvent& MouseEvent )
+{
+	if ( MouseEvent.GetEffectingButton() == EKeys::LeftMouseButton )
+	{
+		OnCurrentThemeChanged.ExecuteIfBound(ColorTheme.Pin());
+
+		return FReply::Handled();
+	}
+	else
+	{
+		return FReply::Unhandled();
+	}
+}
+
+
+
+
+
+TArray< TSharedPtr<FColorTheme> > SColorThemesViewer::ColorThemes;
+TWeakPtr<FColorTheme> SColorThemesViewer::CurrentlySelectedThemePtr;
+bool SColorThemesViewer::bSRGBEnabled = false;
+
+BEGIN_SLATE_FUNCTION_BUILD_OPTIMIZATION
+void SColorThemesViewer::Construct(const FArguments& InArgs)
+{
+	#define LOCTEXT_NAMESPACE "ColorThemesViewer"
+
+	bUseAlpha = InArgs._UseAlpha;
+
+	LoadColorThemesFromIni();
+
+	const FSlateFontInfo SmallLayoutFont( FPaths::EngineContentDir() / TEXT("Slate/Fonts/Roboto-Regular.ttf"), 10 );
+
+	// different menus that could be visible for the color themes menu
+	// standard menu with "new", "rename" and "delete"
+	MenuStandard =
+		SNew(SHorizontalBox)
+		+SHorizontalBox::Slot() .FillWidth(1) .Padding(3)
+		[
+			SNew(SButton) .Text(LOCTEXT("NewButton", "New")) .HAlign(HAlign_Center)
+				.OnClicked(this, &SColorThemesViewer::NewColorTheme)
+		]
+		+SHorizontalBox::Slot() .FillWidth(1) .Padding(3)
+		[
+			SNew(SButton) .Text(LOCTEXT("DuplicateButton", "Duplicate")) .HAlign(HAlign_Center)
+				.OnClicked(this, &SColorThemesViewer::DuplicateColorTheme)
+		]
+		+SHorizontalBox::Slot() .FillWidth(1) .Padding(3)
+		[
+			SNew(SButton) .Text(LOCTEXT("RenameButton", "Rename")) .HAlign(HAlign_Center)
+				.OnClicked(this, &SColorThemesViewer::MenuToRename)
+		]
+		+SHorizontalBox::Slot() .FillWidth(1) .Padding(3)
+		[
+			SNew(SButton) .Text(LOCTEXT("DeleteButton", "Delete")) .HAlign(HAlign_Center)
+				.OnClicked(this, &SColorThemesViewer::MenuToDelete)
+		];
+
+	// menu for renaming the currently selected color theme
+	MenuRename =
+		SNew(SHorizontalBox)
+		+SHorizontalBox::Slot() .FillWidth(1) .Padding(3)
+		[
+			SAssignNew(RenameTextBox, SEditableTextBox) .Font(SmallLayoutFont)
+				.OnTextChanged(this, &SColorThemesViewer::ChangeThemeName)
+		]
+		+SHorizontalBox::Slot().AutoWidth() .Padding(3)
+		[
+			SNew(SButton) .Text(LOCTEXT("AcceptRenameButton", "Accept")) .HAlign(HAlign_Right)
+				.OnClicked(this, &SColorThemesViewer::AcceptThemeName)
+				.IsEnabled(this, &SColorThemesViewer::CanAcceptThemeName)
+		]
+		+SHorizontalBox::Slot().AutoWidth() .Padding(3)
+		[
+			SNew(SButton) .Text(LOCTEXT("CancelRenameButton", "Cancel")) .HAlign(HAlign_Right)
+				.OnClicked(this, &SColorThemesViewer::MenuToStandard)
+		];
+
+	// menu for confirming whether you wish to delete the currently selected color theme
+	MenuConfirmDelete =
+		SNew(SHorizontalBox)
+		+SHorizontalBox::Slot() .FillWidth(1) .HAlign(HAlign_Left) .Padding(3)
+		[
+			SNew(SButton) .Text(LOCTEXT("ConfirmDeleteButton", "Delete")) .HAlign(HAlign_Center)
+				.OnClicked(this, &SColorThemesViewer::DeleteColorTheme)
+		]
+		+SHorizontalBox::Slot() .FillWidth(1) .HAlign(HAlign_Center) .VAlign(VAlign_Center) .Padding(3)
+		[
+			SNew(STextBlock) .Text(LOCTEXT("ConfirmDeletePrompt", "Confirm Delete")) .Font(SmallLayoutFont)
+		]
+		+SHorizontalBox::Slot() .FillWidth(1) .HAlign(HAlign_Right) .Padding(3)
+		[
+			SNew(SButton) .Text(LOCTEXT("CancelDeleteButton", "Cancel")) .HAlign(HAlign_Center)
+				.OnClicked(this, &SColorThemesViewer::MenuToStandard)
+		];
+
+	// menu for dropping colors into the trash
+	MenuTrashColor =
+		SNew(SVerticalBox)
+		+SVerticalBox::Slot().AutoHeight()
+		[
+			SNew(SColorTrash)
+		];
+
+	this->ChildSlot
+	[
+		SNew(SVerticalBox)
+		+SVerticalBox::Slot().AutoHeight()
+		[
+			SNew(SBox)
+			.WidthOverride(320)
+			[
+				SAssignNew(ColorThemeList, SListView< TSharedPtr<FColorTheme> >)
+					.ItemHeight(32)
+					.ListItemsSource(&ColorThemes)
+					.OnGenerateRow(this, &SColorThemesViewer::OnGenerateColorThemeBars)
+			]
+		]
+		+SVerticalBox::Slot().AutoHeight() .Padding(0,15,0,0)
+		[
+			SAssignNew(Menu, SBorder)
+		]
+		+ SVerticalBox::Slot().AutoHeight()
+		[
+			SAssignNew(ErrorText, SErrorText)
+				.Visibility(this, &SColorThemesViewer::OnGetErrorTextVisibility)
+		]
+	];
+
+	MenuToStandardNoReturn();
+
+	if (!CurrentlySelectedThemePtr.IsValid())
+	{
+		CurrentlySelectedThemePtr = ColorThemes[0];
+	}
+	ColorThemeList->ClearSelection();
+	ColorThemeList->SetItemSelection(CurrentlySelectedThemePtr.Pin(), true);
+
+	#undef LOCTEXT_NAMESPACE
+}
+END_SLATE_FUNCTION_BUILD_OPTIMIZATION
+
+void SColorThemesViewer::MenuToStandardNoReturn()
+{
+	ErrorText->SetError(FString());
+	Menu->SetContent(MenuStandard.ToSharedRef());
+}
+
+FReply SColorThemesViewer::MenuToStandard()
+{
+	MenuToStandardNoReturn();
+	return FReply::Handled();
+}
+
+FReply SColorThemesViewer::MenuToRename()
+{
+	RenameTextBox->SetText( FText::FromString( GetCurrentColorTheme()->Name ) );
+	Menu->SetContent(MenuRename.ToSharedRef());
+	return FReply::Handled();
+}
+
+FReply SColorThemesViewer::MenuToDelete()
+{
+	Menu->SetContent(MenuConfirmDelete.ToSharedRef());
+	return FReply::Handled();
+}
+
+void SColorThemesViewer::MenuToTrash()
+{
+	Menu->SetContent(MenuTrashColor.ToSharedRef());
+}
+
+void SColorThemesViewer::RefreshThemes()
+{
+	ColorThemeList->RequestListRefresh();
+
+	SaveColorThemesToIni();
+
+	MenuToStandardNoReturn();
+}
+
+TSharedPtr<FColorTheme> SColorThemesViewer::GetCurrentColorTheme() const
+{
+	return CurrentlySelectedThemePtr.IsValid() ? CurrentlySelectedThemePtr.Pin() : ColorThemes[0];
+}
+
+void SColorThemesViewer::SetUseAlpha( const TAttribute<bool>& InUseAlpha )
+{ 
+	bUseAlpha = InUseAlpha;
+}
+
+void SColorThemesViewer::SetCurrentColorTheme(TSharedPtr<FColorTheme> NewTheme)
+{
+	// Set the current theme, requires a preexisting theme to be passed in
+	CurrentlySelectedThemePtr = NewTheme;
+	ColorThemeList->ClearSelection();
+	ColorThemeList->SetItemSelection(NewTheme, true);
+
+	CurrentThemeChangedEvent.Broadcast();
+	MenuToStandardNoReturn();
+}
+
+TSharedPtr<FColorTheme> SColorThemesViewer::IsColorTheme(const FString& ThemeName)
+{
+	// Find the desired theme
+	for (int32 ThemeIndex = 0; ThemeIndex < ColorThemes.Num(); ++ThemeIndex)
+	{
+		const TSharedPtr<FColorTheme>& ColorTheme = ColorThemes[ThemeIndex];
+		if ( ColorTheme->Name == ThemeName )
+		{
+			return ColorTheme;
+		}
+	}
+	return TSharedPtr<FColorTheme>();
+}
+
+TSharedPtr<FColorTheme> SColorThemesViewer::GetColorTheme(const FString& ThemeName)
+{
+	// Create the desired theme, if not already
+	TSharedPtr<FColorTheme> ColorTheme = IsColorTheme(ThemeName);
+	if ( !ColorTheme.IsValid() )
+	{
+		ColorTheme = NewColorTheme(ThemeName);
+	}
+	return ColorTheme;
+}
+
+FString SColorThemesViewer::MakeUniqueThemeName(const FString& ThemeName)
+{
+	// Ensure the name of the color theme is unique
+	int32 ThemeID = 1;
+	FString NewThemeName = ThemeName;
+	while ( IsColorTheme(NewThemeName).IsValid() )
+	{
+		NewThemeName = ThemeName + FString::Printf( TEXT( " %d" ), ThemeID );
+		ThemeID++;
+	}
+	return NewThemeName;
+}
+
+TSharedPtr<FColorTheme> SColorThemesViewer::NewColorTheme(const FString& ThemeName, const TArray< TSharedPtr<FLinearColor> >& ThemeColors)
+{
+	// Create a uniquely named theme
+	check( ThemeName.Len() > 0 );
+	const FString NewThemeName = MakeUniqueThemeName( ThemeName );
+	ColorThemes.Add(MakeShareable(new FColorTheme( NewThemeName, ThemeColors )));
+	return ColorThemes.Last();
+}
+
+TSharedPtr<FColorTheme> SColorThemesViewer::GetDefaultColorTheme(bool bCreateNew)
+{
+	// Create a default theme (if bCreateNew is always creates a new one, even if there's already a like named theme)
+	const FText Name = NSLOCTEXT("ColorThemesViewer", "NewThemeName", "New Theme");
+	if ( bCreateNew )
+	{
+		return NewColorTheme( Name.ToString() );
+	}
+	return GetColorTheme( Name.ToString() );
+}
+
+FReply SColorThemesViewer::AcceptThemeName()
+{
+	// Update the theme name if it differs, ensuring it is still unique
+	const FString Name = RenameTextBox->GetText().ToString();
+	if ( GetCurrentColorTheme()->Name != Name )
+	{
+		GetCurrentColorTheme()->Name = MakeUniqueThemeName( Name );
+		RefreshThemes();
+	}
+	return FReply::Handled();
+}
+
+bool SColorThemesViewer::CanAcceptThemeName() const
+{
+	return !ErrorText->HasError();
+}
+
+void SColorThemesViewer::ChangeThemeName(const FText& InText)
+{
+	ErrorText->SetError(FString());
+
+	const FString ThemeName = InText.ToString();
+	for (int32 ThemeIndex = 0; ThemeIndex < ColorThemes.Num(); ++ThemeIndex)
+	{
+		const TSharedPtr<FColorTheme>& ColorTheme = ColorThemes[ThemeIndex];
+		if (ColorTheme != GetCurrentColorTheme() && ColorTheme->Name == ThemeName)
+		{
+			FFormatNamedArguments Args;
+			Args.Add(TEXT("Name"), InText);
+			ErrorText->SetError(FText::Format(NSLOCTEXT("ColorThemesViewer", "VerifyTextDup", "A theme already exists with the name '{Name}'."), Args).ToString());
+			return;
+		}
+	}
+}
+
+EVisibility SColorThemesViewer::OnGetErrorTextVisibility() const
+{
+	return !CanAcceptThemeName() ? EVisibility::Visible : EVisibility::Collapsed;
+}
+
+FReply SColorThemesViewer::NewColorTheme()
+{
+	// Create a new, defaultly named theme and update the display
+	GetDefaultColorTheme(true);
+	RefreshThemes();
+	return FReply::Handled();
+}
+
+FReply SColorThemesViewer::DuplicateColorTheme()
+{
+	// Create a copy of the existing current color theme
+	TArray< TSharedPtr<FLinearColor> > NewColors;
+	const TArray< TSharedPtr<FLinearColor> >& CurrentColors = CurrentlySelectedThemePtr.Pin()->GetColors();
+	for (int32 ColorIndex = 0; ColorIndex < CurrentColors.Num(); ++ColorIndex)
+	{
+		NewColors.Add(MakeShareable(new FLinearColor(*CurrentColors[ColorIndex])));
+	}
+	const FText Name = NSLOCTEXT("ColorThemesViewer", "CopyThemeNameAppend", " Copy");
+	NewColorTheme( CurrentlySelectedThemePtr.Pin()->Name + Name.ToString(), NewColors );
+	RefreshThemes();
+	return FReply::Handled();
+}
+
+FReply SColorThemesViewer::DeleteColorTheme()
+{
+	// Delete the current color theme
+	ColorThemes.Remove(GetCurrentColorTheme());
+	if (!ColorThemes.Num())
+	{
+	 	// Create the default if none exists
+		GetDefaultColorTheme();
+	}
+	SetCurrentColorTheme(ColorThemes[0]);
+	RefreshThemes();
+	return FReply::Handled();
+}
+
+bool SColorThemesViewer::OnReadUseSRGB() const
+{
+	return bSRGBEnabled;
+}
+
+bool SColorThemesViewer::OnReadUseAlpha() const
+{
+	return bUseAlpha.Get();
+}
+
+TSharedRef<ITableRow> SColorThemesViewer::OnGenerateColorThemeBars(TSharedPtr<FColorTheme> InItem, const TSharedRef<STableViewBase>& OwnerTable)
+{
+	return
+	SNew( STableRow< TSharedPtr<FColorTheme> >, OwnerTable )
+	[
+		SNew(SColorThemeBar)
+		. ColorTheme(InItem)
+		. OnCurrentThemeChanged(this, &SColorThemesViewer::SetCurrentColorTheme)
+		. ShowTrashCallback(this, &SColorThemesViewer::MenuToTrash)
+		. HideTrashCallback(this, &SColorThemesViewer::MenuToStandardNoReturn)
+		. UseSRGB(this, &SColorThemesViewer::OnReadUseSRGB)
+		. UseAlpha(this, &SColorThemesViewer::OnReadUseAlpha)
+	];
+}
+
+void SColorThemesViewer::LoadColorThemesFromIni()
+{
+	if (FPaths::FileExists(GEditorUserSettingsIni))
+	{
+		bool bThemesRemaining = true;
+		int32 ThemeID = 0;
+		while (bThemesRemaining)
+		{
+			const FString ThemeName = GConfig->GetStr(TEXT("ColorThemes"), *FString::Printf(TEXT("Theme%i"), ThemeID), GEditorUserSettingsIni);
+			if (!ThemeName.IsEmpty())
+			{
+				TSharedPtr<FColorTheme> ColorTheme = GetColorTheme( ThemeName );
+				check( ColorTheme.IsValid() );
+				bool bColorsRemaining = true;
+				int32 ColorID = 0;
+				while (bColorsRemaining)
+				{
+					const FString ColorString = GConfig->GetStr(TEXT("ColorThemes"), *FString::Printf(TEXT("Theme%iColor%i"), ThemeID, ColorID), GEditorUserSettingsIni);
+					if (!ColorString.IsEmpty())
+					{
+						// Add the color if it hasn't already
+						FLinearColor Color;
+						Color.InitFromString(ColorString);
+						if ( ColorTheme->FindApproxColor( Color ) == INDEX_NONE )
+						{
+							ColorTheme->InsertNewColor( MakeShareable(new FLinearColor(Color)), 0 );
+						}
+						++ColorID;
+					}
+					else
+					{
+						bColorsRemaining = false;
+					}
+				}
+				++ThemeID;
+			}
+			else
+			{
+				bThemesRemaining = false;
+			}
+		}
+	}
+	
+	if (!ColorThemes.Num())
+	{
+	 	// Create the default if none exists
+		GetDefaultColorTheme();
+	}
+}
+
+void SColorThemesViewer::SaveColorThemesToIni()
+{
+	if (FPaths::FileExists(GEditorUserSettingsIni))
+	{
+		GConfig->EmptySection(TEXT("ColorThemes"), GEditorUserSettingsIni);
+		for (int32 ThemeIndex = 0; ThemeIndex < ColorThemes.Num(); ++ThemeIndex)
+		{
+			const TSharedPtr<FColorTheme>& Theme = ColorThemes[ThemeIndex];
+			GConfig->SetString(TEXT("ColorThemes"), *FString::Printf(TEXT("Theme%i"), ThemeIndex), *Theme->Name, GEditorUserSettingsIni);
+
+			const TArray< TSharedPtr<FLinearColor> >& Colors = Theme->GetColors();
+			for (int32 ColorIndex = 0; ColorIndex < Colors.Num(); ++ColorIndex)
+			{
+				const TSharedPtr<FLinearColor>& Color = Colors[ColorIndex];
+				GConfig->SetString(TEXT("ColorThemes"), *FString::Printf(TEXT("Theme%iColor%i"), ThemeIndex, ColorIndex), *Color->ToString(), GEditorUserSettingsIni);
+			}
+		}
+	}
+}
diff --git a/Engine/Source/Runtime/AppFramework/Private/Widgets/Testing/STestSuite.cpp b/Engine/Source/Runtime/AppFramework/Private/Widgets/Testing/STestSuite.cpp
--- a/Engine/Source/Runtime/AppFramework/Private/Widgets/Testing/STestSuite.cpp
+++ b/Engine/Source/Runtime/AppFramework/Private/Widgets/Testing/STestSuite.cpp
@@ -22,9 +22,12 @@
 #include "SSearchBox.h"
 #include "SVolumeControl.h"
 #include "SResponsiveGridPanel.h"
+#include "SColorPicker.h"
+
 
 #define LOCTEXT_NAMESPACE "STestSuite"
 
+
 namespace
 {
 	/**
diff --git a/Engine/Source/Runtime/AppFramework/Public/AppFramework.h b/Engine/Source/Runtime/AppFramework/Public/AppFramework.h
deleted file mode 100644
diff --git a/Engine/Source/Runtime/AppFramework/Public/Widgets/Colors/SColorPicker.cpp b/Engine/Source/Runtime/AppFramework/Public/Widgets/Colors/SColorPicker.cpp
deleted file mode 100644
diff --git a/Engine/Source/Runtime/AppFramework/Public/Widgets/Colors/SColorPicker.h b/Engine/Source/Runtime/AppFramework/Public/Widgets/Colors/SColorPicker.h
--- a/Engine/Source/Runtime/AppFramework/Public/Widgets/Colors/SColorPicker.h
+++ b/Engine/Source/Runtime/AppFramework/Public/Widgets/Colors/SColorPicker.h
@@ -2,6 +2,14 @@
 
 #pragma once
 
+#include "SlateBasics.h"
+
+
+// forward declarations
+class SColorThemesViewer;
+class SComboButton;
+class SThemeColorBlocksBar;
+
 
 /** Called when the color picker cancel button is pressed */
 DECLARE_DELEGATE_OneParam(FOnColorPickerCancelled, FLinearColor);
@@ -187,7 +195,7 @@ protected:
 	void ShowSmallTrash();
 
 	/** Cycles the color picker's mode. */
-	void CycleMode( );
+	void CycleMode();
 
 	/**
 	 * Creates a color slider widget for the specified channel.
@@ -210,7 +218,7 @@ protected:
 	 *
 	 * @return The new color preview box.
 	 */
-	TSharedRef<SWidget> MakeColorPreviewBox( ) const;
+	TSharedRef<SWidget> MakeColorPreviewBox() const;
 
 protected:	
 
@@ -230,7 +238,7 @@ private:
 	void HandleAdvancedAreaExpansionChanged( bool Expanded );
 
 	// Callback for getting the visibility of the alpha channel portion in color blocks.
-	EVisibility HandleAlphaColorBlockVisibility( ) const;
+	EVisibility HandleAlphaColorBlockVisibility() const;
 
 	// Callback for clicking the Cancel button.
 	FReply HandleCancelButtonClicked();
@@ -239,7 +247,7 @@ private:
 	FReply HandleColorAreaMouseDown( const FGeometry& MyGeometry, const FPointerEvent& MouseEvent );
 
 	// Callback for clicking the color picker mode button.
-	FReply HandleColorPickerModeButtonClicked( );
+	FReply HandleColorPickerModeButtonClicked();
 
 	// Callback for getting the visibility of the given color picker mode.
 	EVisibility HandleColorPickerModeVisibility( EColorPickerModes Mode ) const;
@@ -260,10 +268,10 @@ private:
 	void HandleColorSpinBoxValueChanged( float NewValue, EColorPickerChannels Channel );
 
 	// Callback for completed eye dropper interactions.
-	void HandleEyeDropperButtonComplete( );
+	void HandleEyeDropperButtonComplete();
 
 	// Callback for getting the text in the hex box.
-	FText HandleHexBoxText( ) const;
+	FText HandleHexBoxText() const;
 
 	// Callback for committed text in the hex input box.
 	void HandleHexInputTextCommitted( const FText& Text, ETextCommit::Type CommitType );
@@ -272,10 +280,10 @@ private:
 	void HandleHSVColorChanged( FLinearColor NewValue );
 
 	// Callback for when interactive user input begins.
-	void HandleInteractiveChangeBegin( );
+	void HandleInteractiveChangeBegin();
 
 	// Callback for when interactive user input ends.
-	void HandleInteractiveChangeEnd( );
+	void HandleInteractiveChangeEnd();
 
 	// Callback for when interactive user input ends.
 	void HandleInteractiveChangeEnd( float NewValue );
@@ -284,13 +292,13 @@ private:
 	FReply HandleNewColorBlockMouseButtonDown( const FGeometry& MyGeometry, const FPointerEvent& MouseEvent, bool bCheckAlpha );
 
 	// Callback for clicking the OK button.
-	FReply HandleOkButtonClicked( );
+	FReply HandleOkButtonClicked();
 
 	// Callback for clicking the old color preview block.
 	FReply HandleOldColorBlockMouseButtonDown( const FGeometry& MyGeometry, const FPointerEvent& MouseEvent, bool bCheckAlpha );
 
 	// Callback for checking whether sRGB colors should be rendered.
-	bool HandleColorPickerUseSRGB( ) const;
+	bool HandleColorPickerUseSRGB() const;
 
 	// Callback for when the parent window has been closed.
 	void HandleParentWindowClosed( const TSharedRef<SWindow>& Window );
@@ -302,22 +310,22 @@ private:
 	void HandleSRGBCheckBoxCheckStateChanged( ESlateCheckBoxState::Type InIsChecked );
 
 	// Callback for determining whether the sRGB check box should be checked.
-	ESlateCheckBoxState::Type HandleSRGBCheckBoxIsChecked( ) const;
+	ESlateCheckBoxState::Type HandleSRGBCheckBoxIsChecked() const;
 
 	// Callback for selecting a color in the color theme bar.
 	void HandleThemeBarColorSelected( FLinearColor NewValue );
 
 	// Callback for getting the theme bar's color theme.
-	TSharedPtr<class FColorTheme> HandleThemeBarColorTheme( ) const;
+	TSharedPtr<class FColorTheme> HandleThemeBarColorTheme() const;
 
 	// Callback for getting the visibility of the theme bar hint text.
-	EVisibility HandleThemeBarHintVisibility( ) const;
+	EVisibility HandleThemeBarHintVisibility() const;
 
 	// Callback for determining whether the theme bar should display the alpha channel.
-	bool HandleThemeBarUseAlpha( ) const;
+	bool HandleThemeBarUseAlpha() const;
 
 	// Callback for theme viewer changes.
-	void HandleThemesViewerThemeChanged( );
+	void HandleThemesViewerThemeChanged();
 
 private:
 	
@@ -388,13 +396,13 @@ private:
 	TWeakPtr<SWindow> ParentWindowPtr;
 
 	/** The widget which holds the currently selected theme */
-	TSharedPtr<class SThemeColorBlocksBar> CurrentThemeBar;
+	TSharedPtr<SThemeColorBlocksBar> CurrentThemeBar;
 
 	/** Widget which is either the button to show the color themes viewer, or to be a color trash */
 	TSharedPtr<SBorder> ColorThemeButtonOrSmallTrash;
 
 	/** The button to show the color themes viewer */
-	TSharedPtr<class SComboButton> ColorThemeComboButton;
+	TSharedPtr<SComboButton> ColorThemeComboButton;
 
 	/** The small color trash shown in place of the combo button */
 	TSharedPtr<SWidget> SmallTrash;
@@ -431,7 +439,7 @@ private:
 private:
 
 	/** A static pointer to the global color themes viewer */
-	static TWeakPtr<class SColorThemesViewer> ColorThemesViewer;
+	static TWeakPtr<SColorThemesViewer> ColorThemesViewer;
 };
 
 
diff --git a/Engine/Source/Runtime/AppFramework/Public/Widgets/Colors/SColorThemes.cpp b/Engine/Source/Runtime/AppFramework/Public/Widgets/Colors/SColorThemes.cpp
deleted file mode 100644
diff --git a/Engine/Source/Runtime/AppFramework/Public/Widgets/Colors/SColorThemes.h b/Engine/Source/Runtime/AppFramework/Public/Widgets/Colors/SColorThemes.h
--- a/Engine/Source/Runtime/AppFramework/Public/Widgets/Colors/SColorThemes.h
+++ b/Engine/Source/Runtime/AppFramework/Public/Widgets/Colors/SColorThemes.h
@@ -2,6 +2,12 @@
 
 #pragma once
 
+#include "SlateBasics.h"
+
+
+// forward declarations
+class SThemeColorBlocksBar;
+
 
 /**
  * A Color Theme is a name and an array of Colors.
@@ -14,7 +20,10 @@ public:
 	FColorTheme(const FString& InName = TEXT(""), const TArray< TSharedPtr<FLinearColor> >& InColors = TArray< TSharedPtr<FLinearColor> >());
 
 	/** Get a list of all the colors in the theme */
-	const TArray< TSharedPtr<FLinearColor> >& GetColors() const { return Colors; }
+	const TArray< TSharedPtr<FLinearColor> >& GetColors() const
+	{
+		return Colors;
+	}
 
 	/** Insert a color at a specific point in the list and broadcast change */
 	void InsertNewColor(TSharedPtr<FLinearColor> InColor, int32 InsertPosition);
@@ -34,7 +43,10 @@ public:
 	FString Name;
 
 	DECLARE_EVENT( FColorTheme, FRefreshEvent );
-	FRefreshEvent& OnRefresh() { return RefreshEvent; }
+	FRefreshEvent& OnRefresh()
+	{
+		return RefreshEvent;
+	}
 
 private:
 
@@ -48,7 +60,8 @@ private:
  * The SColorTrash is a multipurpose widget which allows FColorDragDrops
  * to be dropped on to to be 
  */
-class SColorTrash : public SCompoundWidget
+class SColorTrash
+	: public SCompoundWidget
 {
 public:
 
@@ -98,15 +111,13 @@ private:
 };
 
 
-class SThemeColorBlocksBar;
-
-
 /**
  * SThemeColorBlocks are Color Blocks which point to a Color in a ColorTheme.
  * They can be dragged and dropped, and clicking on one in the Color Picker will
  * give the color that they point to.
  */
-class SThemeColorBlock : public SCompoundWidget
+class SThemeColorBlock
+	: public SCompoundWidget
 {
 public:
 
@@ -205,7 +216,8 @@ private:
 /**
  * SThemeColorBlocksBars are panels for dragging and dropping SColorThemeBlocks onto and off of.
  */
-class SThemeColorBlocksBar : public SPanel
+class SThemeColorBlocksBar
+	: public SPanel
 {
 public:
 
@@ -329,9 +341,8 @@ public:
 	 * This removes a callback from the Color Theme that this bar has a pointer to
 	 */
 	void RemoveRefreshCallback();
-	/**
-	 * Adds a callback to the Color Theme that this bar has a pointer to
-	 */
+
+	/** Adds a callback to the Color Theme that this bar has a pointer to. */
 	void AddRefreshCallback();
 
 	/**
@@ -343,9 +354,7 @@ public:
 	
 private:
 
-	/**
-	 * Destroys the placeholder block in this widget
-	 */
+	/** Destroys the placeholder block in this widget. */
 	void DestroyPlaceholders();
 
 	/** The children blocks of this panel */
@@ -356,10 +365,13 @@ private:
 	 * FColorDragDrop onto this widget. It "replaces" the FColorDragDrop widget with a preview SThemeColorBlock
 	 */
 	TSharedPtr<SThemeColorBlock> NewColorBlockPlaceholder;
+
 	/** The placeholder's color, also used to determine whether the placeholder is real or not (by IsValid) */
 	TSharedPtr<FLinearColor> NewColorPlaceholder;
+
 	/** Current x offset of the placeholder block */
 	float PlaceholderBlockOffset;
+
 	/** The initial grab offset when grabbing the placeholder */
 	FVector2D PlaceholderInitialGrabOffset;
 
@@ -380,6 +392,7 @@ private:
 
 	/** Callback to pass down to the FColorDragDrop for it to show the trash */
 	FSimpleDelegate ShowTrashCallback;
+
 	/** Callback to pass down to the FColorDragDrop for it to hide the trash */
 	FSimpleDelegate HideTrashCallback;
 
@@ -391,7 +404,7 @@ private:
 };
 
 
-DECLARE_DELEGATE_OneParam( FOnCurrentThemeChanged, TSharedPtr<FColorTheme> )
+DECLARE_DELEGATE_OneParam(FOnCurrentThemeChanged, TSharedPtr<FColorTheme>)
 
 
 /**
@@ -456,6 +469,7 @@ private:
 
 	/** Callback to pass down to the FColorDragDrop for it to show the trash */
 	FSimpleDelegate ShowTrashCallback;
+
 	/** Callback to pass down to the FColorDragDrop for it to hide the trash */
 	FSimpleDelegate HideTrashCallback;
 	
@@ -495,12 +509,16 @@ public:
 
 	/** Load the color theme settings from the config */
 	static void LoadColorThemesFromIni();
+
 	/** Save the color theme settings to the config */
 	static void SaveColorThemesToIni();
 	
 	/** Callbacks to execute whenever we change the global current theme */
 	DECLARE_EVENT( SColorThemesViewer, FCurrentThemeChangedEvent );
-	FCurrentThemeChangedEvent& OnCurrentThemeChanged() { return CurrentThemeChangedEvent; }
+	FCurrentThemeChangedEvent& OnCurrentThemeChanged()
+	{
+		return CurrentThemeChangedEvent;
+	}
 
 	/** Whether to use display sRGB color */
 	static bool bSRGBEnabled;
@@ -567,7 +585,8 @@ private:
 	TAttribute<bool> bUseAlpha;
 
 	/** A static holder of the color themes for the entire program */
-	static TArray< TSharedPtr<FColorTheme> > ColorThemes;
+	static TArray<TSharedPtr<FColorTheme>> ColorThemes;
+
 	/** A static pointer to the color theme that is currently selected for the entire program */
 	static TWeakPtr<FColorTheme> CurrentlySelectedThemePtr;
 };
diff --git a/Engine/Source/Runtime/AppFramework/Public/Widgets/SMultipleOptionTable.h b/Engine/Source/Runtime/AppFramework/Public/Widgets/SMultipleOptionTable.h
deleted file mode 100644
diff --git a/Engine/Source/Runtime/AppFramework/Public/Widgets/SWebBrowser.h b/Engine/Source/Runtime/AppFramework/Public/Widgets/SWebBrowser.h
deleted file mode 100644
diff --git a/Engine/Source/Runtime/AppFramework/Public/Widgets/Views/SMultipleOptionTable.h b/Engine/Source/Runtime/AppFramework/Public/Widgets/Views/SMultipleOptionTable.h
new file mode 100644
--- /dev/null
+++ b/Engine/Source/Runtime/AppFramework/Public/Widgets/Views/SMultipleOptionTable.h
@@ -0,0 +1,310 @@
+// Copyright 1998-2014 Epic Games, Inc. All Rights Reserved.
+
+#pragma once
+
+#include "SCheckBox.h"
+#include "SHyperlink.h"
+
+#define LOCTEXT_NAMESPACE "SMultipleOptionTable"
+
+/**
+ * Implements a row widget for an option list.
+ */
+template<typename OptionType>
+class SOptionTableRow
+	: public SMultiColumnTableRow< TSharedPtr<OptionType> >
+{
+public:
+	typedef typename TSlateDelegates< OptionType >::FOnGenerateWidget FOnGenerateWidget;
+	typedef typename STableRow< TSharedPtr<OptionType> >::FArguments FOptionTableRowArgs;
+
+	SLATE_BEGIN_ARGS(SOptionTableRow) { }
+		SLATE_ARGUMENT(TSharedPtr<STableViewBase>, OwnerTableView)
+		SLATE_ARGUMENT(OptionType, Option)
+		SLATE_EVENT(FOnCheckStateChanged, OnCheckStateChanged)
+		SLATE_ATTRIBUTE(ESlateCheckBoxState::Type, IsChecked)
+		SLATE_EVENT(FOnGenerateWidget, OnGenerateWidget)
+	SLATE_END_ARGS()
+
+public:
+
+	/**
+	 * Constructs the widget.
+	 *
+	 * @param InArgs The construction arguments.
+	 */
+	void Construct( const FArguments& InArgs)
+	{
+		Option = InArgs._Option;
+		OnCheckStateChanged = InArgs._OnCheckStateChanged;
+		IsChecked = InArgs._IsChecked;
+		OnGenerateWidget= InArgs._OnGenerateWidget;
+
+		SMultiColumnTableRow< TSharedPtr<OptionType> >::Construct(FOptionTableRowArgs(), InArgs._OwnerTableView.ToSharedRef());
+	}
+
+	/**
+	 * Generates the widget for the specified column.
+	 *
+	 * @param ColumnName The name of the column to generate the widget for.
+	 *
+	 * @return The widget.
+	 */
+	virtual TSharedRef<SWidget> GenerateWidgetForColumn( const FName& ColumnName ) override
+	{
+		if (ColumnName == "IsSelected")
+		{
+			return SAssignNew(CheckBox, SCheckBox)
+				.IsChecked(this, &SOptionTableRow::HandleCheckBoxIsChecked)
+				.OnCheckStateChanged(this, &SOptionTableRow::HandleCheckBoxCheckStateChanged)
+				[
+					OnGenerateWidget.Execute(Option)
+				];
+		}
+
+		return SNullWidget::NullWidget;
+	}
+
+private:
+	// Callback for changing the checked state of the check box.
+	void HandleCheckBoxCheckStateChanged( ESlateCheckBoxState::Type NewState )
+	{
+		OnCheckStateChanged.Execute(NewState);
+	}
+
+	// Callback for determining the checked state of the check box.
+	ESlateCheckBoxState::Type HandleCheckBoxIsChecked( ) const
+	{
+		return IsChecked.Get();
+	}
+
+public:
+	TSharedPtr<SCheckBox> CheckBox;
+
+private:
+	OptionType Option;
+	FOnCheckStateChanged OnCheckStateChanged;
+	TAttribute<ESlateCheckBoxState::Type> IsChecked;
+	FOnGenerateWidget OnGenerateWidget;
+};
+
+template<typename OptionType>
+class SMultipleOptionTable : public SCompoundWidget
+{
+public:
+	typedef typename TSlateDelegates< OptionType >::FOnGenerateWidget FOnGenerateOptionWidget;
+
+	// This callback will be used before selecting or deselecting all of the options. May be of use for optimization.
+	DECLARE_DELEGATE(FOnPreBatchSelect);
+
+	// This callback will be used after selecting or deselecting all of the options. May be of use for optimization.
+	DECLARE_DELEGATE(FOnPostBatchSelect);
+
+	DECLARE_DELEGATE_TwoParams(FOnOptionSelectionChanged, bool, OptionType);
+	DECLARE_DELEGATE_RetVal_OneParam(bool, FIsOptionSelected, OptionType);
+
+	SLATE_BEGIN_ARGS( SMultipleOptionTable<OptionType> ) {}
+		SLATE_EVENT(FOnPreBatchSelect, OnPreBatchSelect)
+		SLATE_EVENT(FOnPostBatchSelect, OnPostBatchSelect)
+		SLATE_EVENT(FOnGenerateOptionWidget, OnGenerateOptionWidget)
+		SLATE_EVENT(FOnOptionSelectionChanged, OnOptionSelectionChanged)
+		SLATE_EVENT(FIsOptionSelected, IsOptionSelected)
+		SLATE_ATTRIBUTE(FOptionalSize, ListHeight)
+	SLATE_END_ARGS()
+
+public:
+	void Construct(const FArguments& InArgs, const TArray< OptionType >* const InOptions);
+
+	// Refreshes the list view. Useful if the options array is modified.
+	void RequestTableRefresh();
+
+private:
+	// Generates the special row widget for this list view widget.
+	TSharedRef<class ITableRow> HandleOptionListViewGenerateRow(OptionType Option, const TSharedRef< class STableViewBase >& OwnerTableView);
+
+	// Calls the provided user callback to generate a widget to display the specified option.
+	TSharedRef<SWidget> GenerateWidgetForOption(OptionType Option);
+
+	// Calls the provided user callback to notify them that an option has been checked or unchecked.
+	void HandleCheckBoxCheckStateChanged(ESlateCheckBoxState::Type NewState, OptionType Option) const;
+
+	// Calls the provided user callback to be notified that an option is selected or unselected in some way.
+	ESlateCheckBoxState::Type HandleCheckBoxIsChecked(OptionType Option) const;
+
+	// Selects all options.
+	void HandleAllHyperlinkNavigate();
+
+	// Deselects all options.
+	void HandleNoneHyperlinkNavigate();
+
+	// Collapses hyperlinks for selecting/deselecting all options based on whether there are multiple options.
+	EVisibility HandleHyperlinkVisibility() const;
+
+private:
+	const TArray< OptionType >* Options;
+	FOnPreBatchSelect OnPreBatchSelect;
+	FOnPostBatchSelect OnPostBatchSelect;
+	FOnGenerateOptionWidget OnGenerateOptionWidget;
+	FOnOptionSelectionChanged OnOptionSelectionChanged;
+	FIsOptionSelected IsOptionSelected;
+	TAttribute<FOptionalSize> ListHeight;
+
+	TSharedPtr< SListView< OptionType > > OptionListView;
+};
+
+template<typename OptionType>
+void SMultipleOptionTable<OptionType>::Construct(const FArguments& InArgs, const TArray< OptionType >* const InOptions)
+{
+	Options = InOptions;
+	OnPreBatchSelect = InArgs._OnPreBatchSelect;
+	OnPostBatchSelect = InArgs._OnPostBatchSelect;
+	OnGenerateOptionWidget = InArgs._OnGenerateOptionWidget;
+	OnOptionSelectionChanged = InArgs._OnOptionSelectionChanged;
+	IsOptionSelected = InArgs._IsOptionSelected;
+	ListHeight = InArgs._ListHeight;
+
+	ChildSlot
+		[
+			SNew(SVerticalBox)
+
+			+ SVerticalBox::Slot()
+			.FillHeight(1.0)
+			.Padding(0.0f, 2.0f, 0.0f, 0.0f)
+			[
+				SNew(SBox)
+				.HeightOverride(ListHeight)
+				[
+					// Options menu
+					SAssignNew(OptionListView, SListView<OptionType>)
+					.HeaderRow
+					(
+					SNew(SHeaderRow)
+					.Visibility(EVisibility::Collapsed)
+
+					+ SHeaderRow::Column("IsSelected")
+					.DefaultLabel(LOCTEXT("OptionListIsSelectedColumnHeader", "IsSelected"))
+					.FillWidth(1.0f)
+					)
+					.ItemHeight(16.0f)
+					.ListItemsSource(Options)
+					.OnGenerateRow(this, &SMultipleOptionTable::HandleOptionListViewGenerateRow)
+					.SelectionMode(ESelectionMode::None)
+				]
+			]
+
+			+ SVerticalBox::Slot()
+				.AutoHeight()
+				.Padding(0.0f, 6.0f, 0.0f, 4.0f)
+				[
+					SNew(SSeparator)
+					.Orientation(Orient_Horizontal)
+				]
+
+			+ SVerticalBox::Slot()
+				.AutoHeight()
+				[
+					SNew(SHorizontalBox)
+
+					+ SHorizontalBox::Slot()
+					.FillWidth(1.0f)
+					.HAlign(HAlign_Right)
+					[
+						SNew(STextBlock)
+						.Text(LOCTEXT("SelectLabel", "Select:"))
+					]
+
+					+ SHorizontalBox::Slot()
+						.AutoWidth()
+						.Padding(8.0f, 0.0f)
+						[
+							// Select all options hyper link
+							SNew(SHyperlink)
+							.OnNavigate(this, &SMultipleOptionTable::HandleAllHyperlinkNavigate)
+							.Text(LOCTEXT("AllHyperlinkLabel", "All"))
+							.ToolTipText(LOCTEXT("AllHyperlinkToolTip", "Select all options."))
+							.Visibility(this, &SMultipleOptionTable::HandleHyperlinkVisibility)									
+						]
+
+					+ SHorizontalBox::Slot()
+						.AutoWidth()
+						[
+							// Deselect all options hyper link
+							SNew(SHyperlink)
+							.OnNavigate(this, &SMultipleOptionTable::HandleNoneHyperlinkNavigate)
+							.Text(LOCTEXT("NoneHyperlinkLabel", "None"))
+							.ToolTipText(LOCTEXT("NoneHyperlinkToolTip", "Deselect all."))
+							.Visibility(this, &SMultipleOptionTable::HandleHyperlinkVisibility)									
+						]
+				]
+		];
+}
+
+template<typename OptionType>
+void SMultipleOptionTable<OptionType>::RequestTableRefresh()
+{
+	OptionListView->RequestListRefresh();
+}
+
+template<typename OptionType>
+TSharedRef<class ITableRow> SMultipleOptionTable<OptionType>::HandleOptionListViewGenerateRow(OptionType Option, const TSharedRef< class STableViewBase >& OwnerTableView)
+{
+	return SNew(SOptionTableRow<OptionType>)
+		.Option(Option)
+		.OwnerTableView(OwnerTableView)
+		.OnCheckStateChanged(this, &SMultipleOptionTable::HandleCheckBoxCheckStateChanged, Option)
+		.IsChecked(this, &SMultipleOptionTable::HandleCheckBoxIsChecked, Option)
+		.OnGenerateWidget(this, &SMultipleOptionTable::GenerateWidgetForOption);
+}
+
+template<typename OptionType>
+TSharedRef<SWidget> SMultipleOptionTable<OptionType>::GenerateWidgetForOption(OptionType Option)
+{
+	return OnGenerateOptionWidget.Execute(Option);
+}
+
+template<typename OptionType>
+void SMultipleOptionTable<OptionType>::HandleCheckBoxCheckStateChanged(ESlateCheckBoxState::Type NewState, OptionType Option) const
+{
+	OnOptionSelectionChanged.Execute(NewState == ESlateCheckBoxState::Checked, Option);
+}
+
+template<typename OptionType>
+ESlateCheckBoxState::Type SMultipleOptionTable<OptionType>::HandleCheckBoxIsChecked(OptionType Option) const
+{
+	return IsOptionSelected.Execute(Option) ? ESlateCheckBoxState::Checked : ESlateCheckBoxState::Unchecked;
+}
+
+template<typename OptionType>
+void SMultipleOptionTable<OptionType>::HandleAllHyperlinkNavigate()
+{
+	OnPreBatchSelect.ExecuteIfBound();
+	for (const OptionType& Option : *Options)
+	{
+		HandleCheckBoxCheckStateChanged(ESlateCheckBoxState::Checked, Option);
+	}
+	OnPostBatchSelect.ExecuteIfBound();
+}
+
+template<typename OptionType>
+void SMultipleOptionTable<OptionType>::HandleNoneHyperlinkNavigate()
+{
+	OnPreBatchSelect.ExecuteIfBound();
+	for (const OptionType& Option : *Options)
+	{
+		HandleCheckBoxCheckStateChanged(ESlateCheckBoxState::Unchecked, Option);
+	}
+	OnPostBatchSelect.ExecuteIfBound();
+}
+
+template<typename OptionType>
+EVisibility SMultipleOptionTable<OptionType>::HandleHyperlinkVisibility() const
+{
+	if (Options->Num() > 1)
+	{
+		return EVisibility::Visible;
+	}
+
+	return EVisibility::Collapsed;
+}
+
+#undef LOCTEXT_NAMESPACE
\ No newline at end of file
diff --git a/Engine/Source/Runtime/AppFramework/Public/Widgets/Web/SWebBrowser.h b/Engine/Source/Runtime/AppFramework/Public/Widgets/Web/SWebBrowser.h
new file mode 100644
--- /dev/null
+++ b/Engine/Source/Runtime/AppFramework/Public/Widgets/Web/SWebBrowser.h
@@ -0,0 +1,88 @@
+// Copyright 1998-2014 Epic Games, Inc. All Rights Reserved.
+
+#pragma once
+
+class IWebBrowserWindow;
+class FWebBrowserViewport;
+
+class APPFRAMEWORK_API SWebBrowser : public SCompoundWidget
+{
+public:
+	SLATE_BEGIN_ARGS(SWebBrowser)
+		: _InitialURL(TEXT("www.google.com"))
+		, _ShowControls(true)
+		, _SupportsTransparency(false)
+		, _ViewportSize(FVector2D(320, 240))
+		{}
+
+		/** A reference to the parent window */
+		SLATE_ARGUMENT(TSharedPtr<SWindow>, ParentWindow)
+
+		/** URL that the browser will initially navigate to */
+		SLATE_ARGUMENT(FString, InitialURL)
+
+		/** Whether to show standard controls like Back, Forward, Reload etc. */
+		SLATE_ARGUMENT(bool, ShowControls)
+
+		/** Should this browser window support transparency */
+		SLATE_ARGUMENT(bool, SupportsTransparency)
+
+		/** Desired size of the web browser viewport */
+		SLATE_ATTRIBUTE(FVector2D, ViewportSize);
+	SLATE_END_ARGS()
+
+	/**
+	 * Constructor
+	 */
+	SWebBrowser();
+
+	/**
+	 * Construct the widget.
+	 *
+	 * @param InArgs  Declaration from which to construct the widget.
+	 */
+	void Construct(const FArguments& InArgs);
+
+	/**
+	 * Get the current title of the web page
+	 */
+	FText GetTitleText() const;
+
+private:
+	/**
+	 * Returns true if the browser can navigate backwards.
+	 */
+	bool CanGoBack() const;
+
+	/**
+	 * Navigate backwards.
+	 */
+	FReply OnBackClicked();
+
+	/**
+	 * Returns true if the browser can navigate forwards.
+	 */
+	bool CanGoForward() const;
+
+	/**
+	 * Navigate forwards.
+	 */
+	FReply OnForwardClicked();
+
+	/**
+	 * Get text for reload button depending on status
+	 */
+	FText GetReloadButtonText() const;
+
+	/**
+	 * Reload or stop loading
+	 */
+	FReply OnReloadClicked();
+
+private:
+	/** Interface for dealing with a web browser window */
+	TSharedPtr<IWebBrowserWindow>	BrowserWindow;
+
+	/** Viewport interface for rendering the web page */
+	TSharedPtr<FWebBrowserViewport>	BrowserViewport;
+};
