Fixed crash when entering an empty string in the hex value of the color picker.

diff --git a/Engine/Source/Runtime/AppFramework/Private/Widgets/Colors/SColorPicker.cpp b/Engine/Source/Runtime/AppFramework/Private/Widgets/Colors/SColorPicker.cpp
--- a/Engine/Source/Runtime/AppFramework/Private/Widgets/Colors/SColorPicker.cpp
+++ b/Engine/Source/Runtime/AppFramework/Private/Widgets/Colors/SColorPicker.cpp
@@ -1262,7 +1262,7 @@ FText SColorPicker::HandleHexBoxText() const
 
 void SColorPicker::HandleHexInputTextCommitted( const FText& Text, ETextCommit::Type CommitType )
 {
-	if ((CommitType == ETextCommit::OnEnter) || (CommitType == ETextCommit::OnUserMovedFocus))
+	if (!Text.IsEmpty() && ((CommitType == ETextCommit::OnEnter) || (CommitType == ETextCommit::OnUserMovedFocus)))
 	{
 		FColor Color = FColor::FromHex(Text.ToString());
 		SetNewTargetColorRGB(FLinearColor(Color.R / 255.0f, Color.G / 255.0f, Color.B / 255.0f, Color.A / 255.0f), false);
diff --git a/Engine/Source/Runtime/Core/Private/Math/Color.cpp b/Engine/Source/Runtime/Core/Private/Math/Color.cpp
--- a/Engine/Source/Runtime/Core/Private/Math/Color.cpp
+++ b/Engine/Source/Runtime/Core/Private/Math/Color.cpp
@@ -135,7 +135,7 @@ float FLinearColor::ComputeLuminance() const
 
 FColor FColor::FromHex( const FString& HexString )
 {
-	int32 StartIndex = (HexString[0] == TCHAR('#')) ? 1 : 0;
+	int32 StartIndex = (!HexString.IsEmpty() && HexString[0] == TCHAR('#')) ? 1 : 0;
 
 	if (HexString.Len() == 3 + StartIndex)
 	{
