Adding a loading spinner to SWebBrowser to be shown before page is first rendered.

Also ignore changes to viewport size that can't be seen, to avoid the whole page being rendered again when it becomes visible.

diff --git a/Engine/Source/Developer/WebBrowser/Private/WebBrowserWindow.cpp b/Engine/Source/Developer/WebBrowser/Private/WebBrowserWindow.cpp
--- a/Engine/Source/Developer/WebBrowser/Private/WebBrowserWindow.cpp
+++ b/Engine/Source/Developer/WebBrowser/Private/WebBrowserWindow.cpp
@@ -10,6 +10,7 @@ FWebBrowserWindow::FWebBrowserWindow(FIntPoint InViewportSize)
 	: UpdatableTexture(nullptr)
 	, ViewportSize(InViewportSize)
 	, bIsClosing(false)
+	, bHasBeenPainted(false)
 {
 	TextureData.Reserve(ViewportSize.X * ViewportSize.Y * 4);
 	TextureData.SetNumZeroed(ViewportSize.X * ViewportSize.Y * 4);
@@ -35,7 +36,9 @@ void FWebBrowserWindow::SetViewportSize(FVector2D WindowSize)
 {
 	// Magic number for texture size, can't access GetMax2DTextureDimension easily
 	FIntPoint ClampedWindowSize = WindowSize.ClampAxes(1, 2048).IntPoint();
-	if (ViewportSize != ClampedWindowSize)
+
+	// Ignore sizes that can't be seen as it forces CEF to re-render whole image
+	if (WindowSize >= FVector2D::UnitVector && ViewportSize != ClampedWindowSize)
 	{
 		FIntPoint OldViewportSize = MoveTemp(ViewportSize);
 		TArray<uint8> OldTextureData = MoveTemp(TextureData);
@@ -76,6 +79,11 @@ bool FWebBrowserWindow::IsValid() const
 	return InternalCefBrowser.get() != nullptr;
 }
 
+bool FWebBrowserWindow::HasBeenPainted() const
+{
+	return bHasBeenPainted;
+}
+
 bool FWebBrowserWindow::IsClosing() const
 {
 	return bIsClosing;
@@ -368,6 +376,8 @@ void FWebBrowserWindow::OnPaint(CefRenderHandler::PaintElementType Type, const C
 	{
 		UpdatableTexture->UpdateTexture(TextureData);
 	}
+
+	bHasBeenPainted = true;
 }
 
 void FWebBrowserWindow::OnCursorChange(CefCursorHandle Cursor)
diff --git a/Engine/Source/Developer/WebBrowser/Private/WebBrowserWindow.h b/Engine/Source/Developer/WebBrowser/Private/WebBrowserWindow.h
--- a/Engine/Source/Developer/WebBrowser/Private/WebBrowserWindow.h
+++ b/Engine/Source/Developer/WebBrowser/Private/WebBrowserWindow.h
@@ -43,6 +43,7 @@ public:
 	virtual void SetViewportSize(FVector2D WindowSize) override;
 	virtual FSlateShaderResource* GetTexture() override;
 	virtual bool IsValid() const override;
+	virtual bool HasBeenPainted() const override;
 	virtual bool IsClosing() const override;
 	virtual FString GetTitle() const override;
 	virtual void OnKeyDown(const FKeyEvent& InKeyEvent) override;
@@ -155,6 +156,8 @@ private:
 	FIntPoint						ViewportSize;
 	/** Whether this window is closing */
 	bool							bIsClosing;
+	/** Whether this window has been painted at least once */
+	bool							bHasBeenPainted;
 
 	// Allow the Handler to access functions only it needs
 	friend class FWebBrowserHandler;
diff --git a/Engine/Source/Developer/WebBrowser/Public/IWebBrowserWindow.h b/Engine/Source/Developer/WebBrowser/Public/IWebBrowserWindow.h
--- a/Engine/Source/Developer/WebBrowser/Public/IWebBrowserWindow.h
+++ b/Engine/Source/Developer/WebBrowser/Public/IWebBrowserWindow.h
@@ -36,6 +36,11 @@ public:
 	 */
 	virtual bool IsValid() const = 0;
 
+	/**
+	 * Checks whether the web browser has been painted at least once
+	 */
+	virtual bool HasBeenPainted() const = 0;
+
 	/**
 	 * Checks whether the web browser is currently being shut down
 	 */
diff --git a/Engine/Source/Runtime/AppFramework/Private/Widgets/SWebBrowser.cpp b/Engine/Source/Runtime/AppFramework/Private/Widgets/SWebBrowser.cpp
--- a/Engine/Source/Runtime/AppFramework/Private/Widgets/SWebBrowser.cpp
+++ b/Engine/Source/Runtime/AppFramework/Private/Widgets/SWebBrowser.cpp
@@ -80,11 +80,25 @@ void SWebBrowser::Construct(const FArguments& InArgs)
 		]
 		+SVerticalBox::Slot()
 		[
-			SAssignNew(ViewportWidget, SViewport)
-			.ViewportSize(InArgs._ViewportSize)
-			.EnableGammaCorrection(false)
-			.EnableBlending(InArgs._SupportsTransparency)
-			.IgnoreTextureAlpha(!InArgs._SupportsTransparency)
+			SNew(SOverlay)
+			+SOverlay::Slot()
+			[
+				SAssignNew(ViewportWidget, SViewport)
+				.ViewportSize(InArgs._ViewportSize)
+				.EnableGammaCorrection(false)
+				.EnableBlending(InArgs._SupportsTransparency)
+				.IgnoreTextureAlpha(!InArgs._SupportsTransparency)
+				.Visibility(this, &SWebBrowser::GetViewportVisibility)
+			]
+			+SOverlay::Slot()
+			.HAlign(HAlign_Center)
+			.VAlign(VAlign_Center)
+			[
+				SNew(SCircularThrobber)
+				.Radius(10.0f)
+				.ToolTipText(LOCTEXT("LoadingThrobberToolTip", "Loading page..."))
+				.Visibility(this, &SWebBrowser::GetLoadingThrobberVisibility)
+			]
 		]
 	];
 
@@ -187,4 +201,26 @@ FReply SWebBrowser::OnReloadClicked()
 	return FReply::Handled();
 }
 
+EVisibility SWebBrowser::GetViewportVisibility() const
+{
+#if WITH_EDITOR || IS_PROGRAM
+	if (BrowserWindow.IsValid() && BrowserWindow->HasBeenPainted())
+	{
+		return EVisibility::Visible;
+	}
+#endif
+	return EVisibility::Hidden;
+}
+
+EVisibility SWebBrowser::GetLoadingThrobberVisibility() const
+{
+#if WITH_EDITOR || IS_PROGRAM
+	if (BrowserWindow.IsValid() && BrowserWindow->HasBeenPainted())
+	{
+		return EVisibility::Hidden;
+	}
+#endif
+	return EVisibility::Visible;
+}
+
 #undef LOCTEXT_NAMESPACE
diff --git a/Engine/Source/Runtime/AppFramework/Public/Widgets/Web/SWebBrowser.h b/Engine/Source/Runtime/AppFramework/Public/Widgets/Web/SWebBrowser.h
--- a/Engine/Source/Runtime/AppFramework/Public/Widgets/Web/SWebBrowser.h
+++ b/Engine/Source/Runtime/AppFramework/Public/Widgets/Web/SWebBrowser.h
@@ -79,6 +79,16 @@ private:
 	 */
 	FReply OnReloadClicked();
 
+	/**
+	 * Get whether the page viewport should be visible
+	 */
+	EVisibility GetViewportVisibility() const;
+
+	/**
+	 * Get whether loading throbber should be visible
+	 */
+	EVisibility GetLoadingThrobberVisibility() const;
+
 private:
 	/** Interface for dealing with a web browser window */
 	TSharedPtr<IWebBrowserWindow>	BrowserWindow;
