Ensure loading spinner isn't shown if page will never load

Reversed the conditions of the check so that the loading spinner is only shown if the browser window is valid and it hasn't been painted yet.

diff --git a/Engine/Source/Runtime/AppFramework/Private/Widgets/SWebBrowser.cpp b/Engine/Source/Runtime/AppFramework/Private/Widgets/SWebBrowser.cpp
--- a/Engine/Source/Runtime/AppFramework/Private/Widgets/SWebBrowser.cpp
+++ b/Engine/Source/Runtime/AppFramework/Private/Widgets/SWebBrowser.cpp
@@ -215,12 +215,12 @@ EVisibility SWebBrowser::GetViewportVisibility() const
 EVisibility SWebBrowser::GetLoadingThrobberVisibility() const
 {
 #if WITH_EDITOR || IS_PROGRAM
-	if (BrowserWindow.IsValid() && BrowserWindow->HasBeenPainted())
+	if (BrowserWindow.IsValid() && !BrowserWindow->HasBeenPainted())
 	{
-		return EVisibility::Hidden;
+		return EVisibility::Visible;
 	}
 #endif
-	return EVisibility::Visible;
+	return EVisibility::Hidden;
 }
 
 #undef LOCTEXT_NAMESPACE
