Slate : Added two events to STableRow: OnCanAcceptDrop and OnAcceptDrop. They are useful for providing an out-of-the-box visualization and handling of DragAndDrop interactions in Lists and Trees. See STableViewTesting for an example.

diff --git a/Engine/Source/Runtime/AppFramework/Private/Framework/Testing/STableViewTesting.cpp b/Engine/Source/Runtime/AppFramework/Private/Framework/Testing/STableViewTesting.cpp
--- a/Engine/Source/Runtime/AppFramework/Private/Framework/Testing/STableViewTesting.cpp
+++ b/Engine/Source/Runtime/AppFramework/Private/Framework/Testing/STableViewTesting.cpp
@@ -6,6 +6,29 @@
 
 #define LOCTEXT_NAMESPACE "STableViewTesting"
 
+class FTestData;
+
+/** Simple DragDropOperation for reordering items in this example's lists and trees. */
+class FTableViewDragDrop : public FDragDropOperation
+{
+public:
+	DRAG_DROP_OPERATOR_TYPE(FTableViewDragDrop, FDragDropOperation);
+
+	static TSharedRef<FTableViewDragDrop> New(const TSharedRef<FTestData>& InTestData)
+	{
+		return MakeShareable(new FTableViewDragDrop(InTestData));
+	}
+
+	const TSharedRef<FTestData>& GetDraggedItem() const { return TestData; }
+
+private:
+	FTableViewDragDrop(const TSharedRef<FTestData>& InTestData)
+	: TestData( InTestData )
+	{
+	}
+
+	TSharedRef<FTestData> TestData;
+};
 
 /**
  * A data item with which we can test lists and trees.
@@ -19,6 +42,71 @@ class FTestData
 			return MakeShareable( new FTestData( ChildName, MakeRandomItemHeight() ) );
 		}
 
+		/**
+		 * Recursively look for `ItemToRemove` in `RemoveFrom` and any of the descendants, and remove `ItemToRemove`.
+		 * @return true when successful.
+		 */
+		static bool RemoveRecursive(TArray< TSharedPtr< FTestData > >& RemoveFrom, const TSharedPtr<FTestData>& ItemToRemove)
+		{
+			const int32 ItemIndex = RemoveFrom.Find(ItemToRemove);
+			if (ItemIndex != INDEX_NONE)
+			{
+				RemoveFrom.RemoveAt(ItemIndex);
+				return true;
+			}
+
+			// Did not successfully remove an item. Try all the children.
+			for (int32 ItemIndex = 0; ItemIndex < RemoveFrom.Num(); ++ItemIndex)
+			{
+				if (RemoveRecursive(RemoveFrom[ItemIndex]->Children, ItemToRemove))
+				{
+					return true;
+				}
+			}
+
+			return false;
+		}
+
+		/**
+		 * Recursively look for `TargetItem` in `InsertInto` and any of the descendants.
+		 * Insert `ItemToInsert` relative to the target.
+		 * Relative positioning dictated by `RelativeLocation`.
+		 *
+		 * @return true when successful.
+		 */
+		static bool InsertRecursive(TArray< TSharedPtr< FTestData > >& InsertInto, const TSharedRef<FTestData>& ItemToInsert, const TSharedRef<FTestData>& TargetItem, EItemDropZone RelativeLocation)
+		{
+			const int32 TargetIndex = InsertInto.Find(TargetItem);
+			if (TargetIndex != INDEX_NONE)
+			{
+				if (RelativeLocation == EItemDropZone::AboveItem)
+				{
+					InsertInto.Insert(ItemToInsert, TargetIndex);
+				}
+				else if (RelativeLocation == EItemDropZone::BelowItem)
+				{
+					InsertInto.Insert(ItemToInsert, TargetIndex + 1);
+				}
+				else
+				{
+					ensure(RelativeLocation == EItemDropZone::OntoItem);
+					InsertInto[TargetIndex]->Children.Insert(ItemToInsert, 0);
+				}				
+				return true;
+			}
+
+			// Did not successfully remove an item. Try all the children.
+			for (int32 ItemIndex = 0; ItemIndex < InsertInto.Num(); ++ItemIndex)
+			{
+				if (InsertRecursive(InsertInto[ItemIndex]->Children, ItemToInsert, TargetItem, RelativeLocation))
+				{
+					return true;
+				}
+			}
+
+			return false;
+		}
+
 		static float MakeRandomItemHeight()
 		{
 			return static_cast<float>(rand()) / RAND_MAX * 50.0f;
@@ -243,7 +331,10 @@ public:
 		: _ItemToEdit()
 		{}
 
-		SLATE_ARGUMENT( TSharedPtr<FTestData>, ItemToEdit )
+		SLATE_EVENT(FOnCanAcceptDrop, OnCanAcceptDrop)
+		SLATE_EVENT(FOnAcceptDrop, OnAcceptDrop)
+		SLATE_EVENT(FOnDragDetected, OnDragDetected)
+		SLATE_ARGUMENT(TSharedPtr<FTestData>, ItemToEdit)
 
 	SLATE_END_ARGS()
 	
@@ -340,7 +431,12 @@ public:
 	{
 		ItemToEdit = InArgs._ItemToEdit;
 		
-		FSuperRowType::Construct( FSuperRowType::FArguments().Padding(1), InOwnerTableView );	
+		FSuperRowType::Construct( FSuperRowType::FArguments()
+			.OnCanAcceptDrop(InArgs._OnCanAcceptDrop)
+			.OnAcceptDrop(InArgs._OnAcceptDrop)
+			.OnDragDetected(InArgs._OnDragDetected)
+			.Padding(1)
+		, InOwnerTableView);
 	}
 	
 	void SpawnContextMenu( const FVector2D& SpawnLocation )
@@ -410,6 +506,8 @@ private:
 	TSharedPtr<FTestData> ItemToEdit;
 };
 
+
+
 /**
  * An Item Editor used by tile view testing.
  * It visualizes a string and edits its contents.
@@ -422,7 +520,10 @@ public:
 		: _ItemToEdit()
 		{}
 
-		SLATE_ARGUMENT( TSharedPtr<FTestData>, ItemToEdit )
+		SLATE_EVENT(FOnCanAcceptDrop, OnCanAcceptDrop)
+		SLATE_EVENT(FOnAcceptDrop, OnAcceptDrop)
+		SLATE_EVENT(FOnDragDetected, OnDragDetected)
+		SLATE_ARGUMENT(TSharedPtr<FTestData>, ItemToEdit)
 
 	SLATE_END_ARGS()
 	
@@ -433,11 +534,12 @@ public:
 	 */
 	void Construct( const FArguments& InArgs, const TSharedRef<STableViewBase>& InOwnerTableView )
 	{
-		STableRow<TSharedPtr<FTestData> >::ConstructInternal(
-			STableRow::FArguments()
-				.ShowSelection(true),
-			InOwnerTableView
-		);
+		STableRow<TSharedPtr<FTestData> >::ConstructInternal( STableRow::FArguments()
+			.OnCanAcceptDrop(InArgs._OnCanAcceptDrop)
+			.OnAcceptDrop(InArgs._OnAcceptDrop)
+			.OnDragDetected(InArgs._OnDragDetected)
+			.ShowSelection(true)
+		, InOwnerTableView);
 
 		ItemToEdit = InArgs._ItemToEdit;
 
@@ -674,13 +776,15 @@ public:
 			+SVerticalBox::Slot()
 			.FillHeight(1)
 			[
-				SNew(SBox) .HeightOverride(500.0f)
+				SNew(SBox)
+				.HeightOverride(500.0f)
 				[
 					SNew( SSplitter )
 					+ SSplitter::Slot()
 					. Value(1)
 					[
 						SNew(SBorder)
+						.BorderImage(FCoreStyle::Get().GetBrush("ToolPanel.GroupBorder"))
 						[
 							SNew(SHorizontalBox)
 							+SHorizontalBox::Slot()
@@ -728,6 +832,7 @@ public:
 					. Value(1)
 					[
 						SNew(SBorder)
+						.BorderImage(FCoreStyle::Get().GetBrush("ToolPanel.GroupBorder"))
 						[
 							// The tile view being tested
 							SAssignNew( TileViewBeingTested, STileView< TSharedPtr<FTestData> > )
@@ -749,6 +854,7 @@ public:
 					. Value(1)
 					[
 						SNew(SBorder)
+						.BorderImage(FCoreStyle::Get().GetBrush("ToolPanel.GroupBorder"))
 						[
 							// The TreeView being tested; mostly identical to ListView except for OnGetChildren
 							SAssignNew( TreeBeingTested, STreeView< TSharedPtr<FTestData> > )
@@ -799,6 +905,45 @@ public:
 	
 protected:
 	
+	FReply OnDragDetected_Handler(const FGeometry&, const FPointerEvent&, TWeakPtr<FTestData> TestData)
+	{
+		return FReply::Handled().BeginDragDrop(FTableViewDragDrop::New(TestData.Pin().ToSharedRef()));
+	}
+
+
+	TOptional<EItemDropZone> OnCanAcceptDrop_Handler(const FDragDropEvent& DragDropEvent, EItemDropZone DropZone, TSharedPtr<FTestData> TargetItem)
+	{
+		TSharedPtr<FTableViewDragDrop> DragDropOperation = DragDropEvent.GetOperationAs<FTableViewDragDrop>();
+		if (DragDropOperation.IsValid())
+		{
+			return DropZone;
+		}
+		else
+		{
+			return TOptional<EItemDropZone>();
+		}
+	}
+
+
+	FReply OnAcceptDrop_Handler(const FDragDropEvent& DragDropEvent, EItemDropZone DropZone, TSharedPtr<FTestData> TargetItem)
+	{
+		TSharedPtr<FTableViewDragDrop> DragDropOperation = DragDropEvent.GetOperationAs<FTableViewDragDrop>();
+		if (DragDropOperation.IsValid())
+		{
+			const TSharedRef<FTestData>& ItemBeingDragged = DragDropOperation->GetDraggedItem();
+			FTestData::RemoveRecursive(Items, ItemBeingDragged);
+			FTestData::InsertRecursive(Items, ItemBeingDragged, TargetItem.ToSharedRef(), DropZone);
+			this->ListBeingTested->RequestListRefresh();
+			this->TreeBeingTested->RequestTreeRefresh();
+			this->TileViewBeingTested->RequestListRefresh();
+			return FReply::Handled();
+		}
+		else
+		{
+			return FReply::Unhandled();
+		}		
+	}
+
 	TSharedPtr<SWidget> GetListContextMenu()
 	{
 		return 
@@ -985,7 +1130,17 @@ protected:
 	{
 		return
 			SNew( SItemEditor, OwnerTable )
-			. ItemToEdit( InItem );
+			// Triggered when a user is attempting to drag, so initiate a DragDropOperation.
+			.OnDragDetected(this, &STableViewTesting::OnDragDetected_Handler, TWeakPtr<FTestData>(InItem))
+
+			// Given a hovered drop zone (above, onto, or below) respond with a zone where we would actually drop the item.
+			// Respond with an un-set TOptional<EItemDropZone> if we cannot drop here at all.
+			.OnCanAcceptDrop(this, &STableViewTesting::OnCanAcceptDrop_Handler)
+
+			// Actually perform the drop into the given drop zone.
+			.OnAcceptDrop(this, &STableViewTesting::OnAcceptDrop_Handler)
+
+			.ItemToEdit( InItem );
 	}
 
 
@@ -996,7 +1151,10 @@ protected:
 	{
 		return
 			SNew( STileItemEditor, OwnerTable )
-			. ItemToEdit( InItem );
+			.OnCanAcceptDrop(this, &STableViewTesting::OnCanAcceptDrop_Handler)
+			.OnAcceptDrop(this, &STableViewTesting::OnAcceptDrop_Handler)
+			.OnDragDetected(this, &STableViewTesting::OnDragDetected_Handler, TWeakPtr<FTestData>(InItem))
+			.ItemToEdit( InItem );
 	}
 	
 	
@@ -1007,7 +1165,10 @@ protected:
 	{
 		return
 			SNew( SItemEditor, OwnerTable )
-			. ItemToEdit( InItem );
+			.OnCanAcceptDrop( this, &STableViewTesting::OnCanAcceptDrop_Handler )
+			.OnAcceptDrop( this, &STableViewTesting::OnAcceptDrop_Handler )
+			.OnDragDetected(this, &STableViewTesting::OnDragDetected_Handler, TWeakPtr<FTestData>(InItem))
+			.ItemToEdit( InItem );
 	}
 	
 	/** Given a data item populate the OutChildren array with the item's children. */
diff --git a/Engine/Source/Runtime/Slate/Public/Widgets/Views/STableRow.h b/Engine/Source/Runtime/Slate/Public/Widgets/Views/STableRow.h
--- a/Engine/Source/Runtime/Slate/Public/Widgets/Views/STableRow.h
+++ b/Engine/Source/Runtime/Slate/Public/Widgets/Views/STableRow.h
@@ -40,6 +40,16 @@ class SLATE_API ITableRow
 };
 
 
+/**
+ * Where we are going to drop relative to the target item.
+ */
+enum class EItemDropZone
+{
+	AboveItem,
+	OntoItem,
+	BelowItem
+};
+
 template <typename ItemType> class SListView;
 
 DECLARE_DELEGATE_OneParam(FOnTableRowDragEnter, FDragDropEvent const&);
@@ -57,6 +67,12 @@ class STableRow : public ITableRow, public SBorder
 {
 	static_assert(TIsValidListItem<ItemType>::Value, "Item type T must be a pointer or a TSharedPtr.");
 
+public:
+	/** Delegate signature for querying whether this FDragDropEvent will be handled by the drop target of type ItemType. */
+	DECLARE_DELEGATE_RetVal_ThreeParams(TOptional<EItemDropZone>, FOnCanAcceptDrop, const FDragDropEvent&, EItemDropZone, ItemType);
+	/** Delegate signature for handling the drop of FDragDropEvent onto target of type ItemType */
+	DECLARE_DELEGATE_RetVal_ThreeParams(FReply, FOnAcceptDrop, const FDragDropEvent&, EItemDropZone, ItemType);
+
 public:
 
 	SLATE_BEGIN_ARGS( STableRow< ItemType > )
@@ -68,7 +84,27 @@ public:
 	
 		SLATE_STYLE_ARGUMENT( FTableRowStyle, Style )
 
-		SLATE_EVENT( FOnDragDetected, OnDragDetected )
+		// High Level DragAndDrop
+
+		/**
+		 * Handle this event to determine whether a drag and drop operation can be executed on top of the target row widget.
+		 * Most commonly, this is used for previewing re-ordering and re-organization operations in lists or trees.
+		 * e.g. A user is dragging one item into a different spot in the list or tree.
+		 *      This delegate will be called to figure out if we should give visual feedback on whether an item will 
+		 *      successfully drop into the list.
+		 */
+		SLATE_EVENT( FOnCanAcceptDrop, OnCanAcceptDrop )
+
+		/**
+		 * Perform a drop operation onto the target row widget
+		 * Most commonly used for executing a re-ordering and re-organization operations in lists or trees.
+		 * e.g. A user was dragging one item into a different spot in the list; they just dropped it.
+		 *      This is our chance to handle the drop by reordering items and calling for a list refresh.
+		 */
+		SLATE_EVENT( FOnAcceptDrop,    OnAcceptDrop )
+
+		// Low level DragAndDrop
+		SLATE_EVENT( FOnDragDetected,      OnDragDetected )
 		SLATE_EVENT( FOnTableRowDragEnter, OnDragEnter )
 		SLATE_EVENT( FOnTableRowDragLeave, OnDragLeave )
 		SLATE_EVENT( FOnTableRowDrop,      OnDrop )
@@ -140,8 +176,6 @@ public:
 
 	virtual int32 OnPaint( const FPaintArgs& Args, const FGeometry& AllottedGeometry, const FSlateRect& MyClippingRect, FSlateWindowElementList& OutDrawElements, int32 LayerId, const FWidgetStyle& InWidgetStyle, bool bParentEnabled ) const override
 	{
-		int32 Result = SBorder::OnPaint(Args, AllottedGeometry, MyClippingRect, OutDrawElements, LayerId, InWidgetStyle, bParentEnabled);
-
 		TSharedRef< ITypedTableView<ItemType> > OwnerWidget = OwnerTablePtr.Pin().ToSharedRef();
 		const bool bIsActive = OwnerWidget->AsWidget()->HasKeyboardFocus();
 		const ItemType* MyItem = OwnerWidget->Private_ItemFromWidget( this );
@@ -158,7 +192,35 @@ public:
 				);
 		}
 
-		return Result;
+		LayerId = SBorder::OnPaint(Args, AllottedGeometry, MyClippingRect, OutDrawElements, LayerId, InWidgetStyle, bParentEnabled);
+
+		if (ItemDropZone.IsSet())
+		{
+			// Draw feedback for user dropping an item above, below, or onto a row.
+			const FSlateBrush* DropIndicatorBrush = [&]()
+			{
+				switch (ItemDropZone.GetValue())
+				{
+					case EItemDropZone::AboveItem: return &Style->DropIndicator_Above; break;
+					default:
+					case EItemDropZone::OntoItem: return &Style->DropIndicator_Onto; break;
+					case EItemDropZone::BelowItem: return &Style->DropIndicator_Below; break;
+				};
+			}();
+
+			FSlateDrawElement::MakeBox
+			(
+				OutDrawElements,
+				LayerId++,
+				AllottedGeometry.ToPaintGeometry(),
+				DropIndicatorBrush,
+				MyClippingRect,
+				ESlateDrawEffect::None,
+				Style->SelectorFocusedBrush.GetTint(InWidgetStyle) * InWidgetStyle.GetColorAndOpacityTint()
+			);
+		}
+
+		return LayerId;
 	}
 
 	/**
@@ -494,34 +556,103 @@ public:
 		{
 			OnDragEnter_Handler.Execute(DragDropEvent);
 		}
-		else
-		{
-			SBorder::OnDragEnter(MyGeometry, DragDropEvent);
-		}
 	}
 
 	virtual void OnDragLeave(FDragDropEvent const& DragDropEvent) override
 	{
+		ItemDropZone = TOptional<EItemDropZone>();
+
 		if (OnDragLeave_Handler.IsBound())
 		{
 			OnDragLeave_Handler.Execute(DragDropEvent);
 		}
+	}
+
+	/** @return the zone (above, onto, below) based on where the user is hovering over within the row */
+	EItemDropZone ZoneFromPointerPosition(FVector2D LocalPointerPos, float RowHeight)
+	{
+		const float VecticalZoneBoundarySu = FMath::Clamp(RowHeight * 0.25f, 3.0f, 10.0f);
+		if (LocalPointerPos.Y < VecticalZoneBoundarySu)
+		{
+			return EItemDropZone::AboveItem;
+		}
+		else if (LocalPointerPos.Y > RowHeight - VecticalZoneBoundarySu)
+		{
+			return EItemDropZone::BelowItem;
+		}
 		else
 		{
-			SBorder::OnDragLeave(DragDropEvent);
+			return EItemDropZone::OntoItem;
 		}
 	}
 
-	virtual FReply OnDrop(FGeometry const& MyGeometry, FDragDropEvent const& DragDropEvent) override
+	virtual FReply OnDragOver(const FGeometry& MyGeometry, const FDragDropEvent& DragDropEvent) override
 	{
-		if (OnDrop_Handler.IsBound())
+		if ( OnCanAcceptDrop.IsBound() )
 		{
-			return OnDrop_Handler.Execute(DragDropEvent);
+			const FVector2D LocalPointerPos = MyGeometry.AbsoluteToLocal(DragDropEvent.GetScreenSpacePosition());
+			const EItemDropZone ItemHoverZone = ZoneFromPointerPosition(LocalPointerPos, MyGeometry.GetLocalSize().Y);
+
+			ItemDropZone = [&]()
+			{
+				TSharedPtr< ITypedTableView<ItemType> > OwnerWidget = OwnerTablePtr.Pin();
+				const ItemType* MyItem = OwnerWidget->Private_ItemFromWidget(this);
+				return OnCanAcceptDrop.Execute(DragDropEvent, ItemHoverZone, *MyItem);
+			}();
+
+			return FReply::Handled();
 		}
 		else
 		{
-			return SBorder::OnDrop(MyGeometry, DragDropEvent);
+			return FReply::Unhandled();
 		}
+
+	}
+
+	virtual FReply OnDrop(const FGeometry& MyGeometry, const FDragDropEvent& DragDropEvent) override
+	{
+		const FReply Reply = [&]()
+		{
+			if (OnAcceptDrop.IsBound())
+			{
+				const TSharedRef< ITypedTableView<ItemType> > OwnerWidget = OwnerTablePtr.Pin().ToSharedRef();
+
+				// A drop finishes the drag/drop operation, so we are no longer providing any feedback.
+				ItemDropZone = TOptional<EItemDropZone>();
+
+				// Find item associated with this widget.
+				const ItemType* MyItem = OwnerWidget->Private_ItemFromWidget(this);
+				
+				// Which physical drop zone is the drop about to be performed onto?
+				const FVector2D LocalPointerPos = MyGeometry.AbsoluteToLocal(DragDropEvent.GetScreenSpacePosition());
+				const EItemDropZone HoveredZone = ZoneFromPointerPosition(LocalPointerPos, MyGeometry.GetLocalSize().Y);
+
+				// The row gets final say over which zone to drop onto regardless of physical location.
+				const TOptional<EItemDropZone> ReportedZone = OnCanAcceptDrop.IsBound()
+					? OnCanAcceptDrop.Execute(DragDropEvent, HoveredZone, *MyItem)
+					: HoveredZone;
+
+				if (ReportedZone.IsSet())
+				{
+					FReply DropReply = OnAcceptDrop.Execute(DragDropEvent, ReportedZone.GetValue(), *MyItem);
+					if (DropReply.IsEventHandled())
+					{
+						// Expand the drop target just in case, so that what we dropped is visible.
+						OwnerWidget->Private_SetItemExpansion(*MyItem, true);
+					}
+				}
+			}
+
+			return FReply::Unhandled();
+		}();
+
+		// @todo slate : Made obsolete by OnAcceptDrop. Get rid of this.
+		if ( !Reply.IsEventHandled() && OnDrop_Handler.IsBound() )
+		{
+			return OnDrop_Handler.Execute(DragDropEvent);
+		}
+
+		return Reply;
 	}
 
 	virtual void SetIndexInList( int32 InIndexInList ) override
@@ -691,6 +822,9 @@ protected:
 
 		this->ForegroundColor = TAttribute<FSlateColor>( this, &STableRow::GetForegroundBasedOnSelection );
 
+		this->OnCanAcceptDrop = InArgs._OnCanAcceptDrop;
+		this->OnAcceptDrop = InArgs._OnAcceptDrop;
+
 		this->OnDragDetected_Handler = InArgs._OnDragDetected;
 		this->OnDragEnter_Handler = InArgs._OnDragEnter;
 		this->OnDragLeave_Handler = InArgs._OnDragLeave;
@@ -743,6 +877,15 @@ protected:
 	/** Style used to draw this table row */
 	const FTableRowStyle* Style;
 
+	/** @see STableRow's OnCanAcceptDrop event */
+	FOnCanAcceptDrop OnCanAcceptDrop;
+
+	/** @see STableRow's OnAcceptDrop event */
+	FOnAcceptDrop OnAcceptDrop;
+
+	/** Are we currently dragging/dropping over this item? */
+	TOptional<EItemDropZone> ItemDropZone;
+
 	/** Delegate triggered when a user starts to drag a list item */
 	FOnDragDetected OnDragDetected_Handler;
 
@@ -795,6 +938,8 @@ protected:
 			.Style(InArgs._Style)
 			.Padding(InArgs._Padding)
 			.ShowSelection(InArgs._ShowSelection)
+			.OnCanAcceptDrop(InArgs._OnCanAcceptDrop)
+			.OnAcceptDrop(InArgs._OnAcceptDrop)
 			.OnDragDetected(InArgs._OnDragDetected)
 			.OnDragEnter(InArgs._OnDragEnter)
 			.OnDragLeave(InArgs._OnDragLeave)
diff --git a/Engine/Source/Runtime/SlateCore/Private/Styling/CoreStyle.cpp b/Engine/Source/Runtime/SlateCore/Private/Styling/CoreStyle.cpp
--- a/Engine/Source/Runtime/SlateCore/Private/Styling/CoreStyle.cpp
+++ b/Engine/Source/Runtime/SlateCore/Private/Styling/CoreStyle.cpp
@@ -652,7 +652,10 @@ TSharedRef<ISlateStyle> FCoreStyle::Create( const FName& InStyleSetName )
 			.SetInactiveHoveredBrush( IMAGE_BRUSH( "Common/Selection", Icon8x8, SelectionColor_Inactive ) )
 			.SetTextColor( DefaultForeground )
 			.SetSelectedTextColor( InvertedForeground )
-			);
+			.SetDropIndicator_Above(BOX_BRUSH("Common/DropZoneIndicator_Above", FMargin(10.0f/16.0f, 10.0f/16.0f, 0,0), SelectionColor ))
+			.SetDropIndicator_Onto(BOX_BRUSH("Common/DropZoneIndicator_Onto", FMargin(4.0f / 16.0f), SelectionColor))
+			.SetDropIndicator_Below(BOX_BRUSH("Common/DropZoneIndicator_Below", FMargin(10.0f / 16.0f, 0, 0, 10.0f / 16.0f), SelectionColor))
+		);
 
 		Style->Set( "TreeArrow_Collapsed", new IMAGE_BRUSH( "Common/TreeArrow_Collapsed", Icon10x10, DefaultForeground ) );
 		Style->Set( "TreeArrow_Collapsed_Hovered", new IMAGE_BRUSH( "Common/TreeArrow_Collapsed_Hovered", Icon10x10, DefaultForeground ) );
diff --git a/Engine/Source/Runtime/SlateCore/Private/Styling/SlateTypes.cpp b/Engine/Source/Runtime/SlateCore/Private/Styling/SlateTypes.cpp
--- a/Engine/Source/Runtime/SlateCore/Private/Styling/SlateTypes.cpp
+++ b/Engine/Source/Runtime/SlateCore/Private/Styling/SlateTypes.cpp
@@ -493,6 +493,9 @@ void FTableRowStyle::GetResources( TArray< const FSlateBrush* >& OutBrushes ) co
 	OutBrushes.Add( &EvenRowBackgroundBrush );
 	OutBrushes.Add( &OddRowBackgroundHoveredBrush );
 	OutBrushes.Add( &OddRowBackgroundBrush );
+	OutBrushes.Add( &DropIndicator_Above );
+	OutBrushes.Add( &DropIndicator_Onto );
+	OutBrushes.Add( &DropIndicator_Below );
 }
 
 const FName FTableRowStyle::TypeName( TEXT("FTableRowStyle") );
diff --git a/Engine/Source/Runtime/SlateCore/Public/Styling/SlateTypes.h b/Engine/Source/Runtime/SlateCore/Public/Styling/SlateTypes.h
--- a/Engine/Source/Runtime/SlateCore/Public/Styling/SlateTypes.h
+++ b/Engine/Source/Runtime/SlateCore/Public/Styling/SlateTypes.h
@@ -1075,11 +1075,6 @@ struct SLATECORE_API FTableRowStyle : public FSlateWidgetStyle
 	FSlateBrush DragDropLineIndicatorBrush;
 	FTableRowStyle& SetDragDropLineIndicatorBrush(const FSlateBrush& InDragDropLineIndicatorBrush){ DragDropLineIndicatorBrush = InDragDropLineIndicatorBrush; return *this; }
 
-	/** Brush used when drawing the drag-drop indicator border around a row */
-	UPROPERTY(EditAnywhere, Category = Appearance)
-	FSlateBrush DragDropBorderIndicatorBrush;
-	FTableRowStyle& SetDragDropBorderIndicatorBrush(const FSlateBrush& InDragDropBorderIndicatorBrush){ DragDropBorderIndicatorBrush = InDragDropBorderIndicatorBrush; return *this; }
-
 	/** Text color used for all rows */
 	UPROPERTY(EditAnywhere, BlueprintReadWrite, Category=Appearance)
 	FSlateColor TextColor;
@@ -1089,6 +1084,21 @@ struct SLATECORE_API FTableRowStyle : public FSlateWidgetStyle
 	UPROPERTY(EditAnywhere, BlueprintReadWrite, Category=Appearance)
 	FSlateColor SelectedTextColor;
 	FTableRowStyle& SetSelectedTextColor( const FSlateColor& InSelectedTextColor ){ SelectedTextColor = InSelectedTextColor; return *this; }
+
+	/** Brush used to provide feedback that a user can drop above the hovered row. */
+	UPROPERTY(EditAnywhere, BlueprintReadWrite, Category = Appearance)
+	FSlateBrush DropIndicator_Above;
+	FTableRowStyle& SetDropIndicator_Above(const FSlateBrush& InValue){ DropIndicator_Above = InValue; return *this; }
+
+	/** Brush used to provide feedback that a user can drop onto the hovered row. */
+	UPROPERTY(EditAnywhere, BlueprintReadWrite, Category = Appearance)
+	FSlateBrush DropIndicator_Onto;
+	FTableRowStyle& SetDropIndicator_Onto(const FSlateBrush& InValue){ DropIndicator_Onto = InValue; return *this; }
+
+	/** Brush used to provide feedback that a user can drop below the hovered row. */
+	UPROPERTY(EditAnywhere, BlueprintReadWrite, Category = Appearance)
+	FSlateBrush DropIndicator_Below;
+	FTableRowStyle& SetDropIndicator_Below(const FSlateBrush& InValue){ DropIndicator_Below = InValue; return *this; }
 };
 
 
