Removed SLATE_TEXT_ARGUMENT from AppFramework specific widgets

diff --git a/Engine/Source/Runtime/AppFramework/Private/Widgets/Colors/SColorThemes.cpp b/Engine/Source/Runtime/AppFramework/Private/Widgets/Colors/SColorThemes.cpp
--- a/Engine/Source/Runtime/AppFramework/Private/Widgets/Colors/SColorThemes.cpp
+++ b/Engine/Source/Runtime/AppFramework/Private/Widgets/Colors/SColorThemes.cpp
@@ -718,7 +718,7 @@ void SColorThemeBar::Construct(const FArguments& InArgs)
 							.ColorTheme(InArgs._ColorTheme)
 							.ShowTrashCallback(ShowTrashCallback)
 							.HideTrashCallback(HideTrashCallback)
-							.EmptyText(NSLOCTEXT("ColorThemesViewer", "NoColorsText", "No Colors Added Yet").ToString())
+							.EmptyText(NSLOCTEXT("ColorThemesViewer", "NoColorsText", "No Colors Added Yet"))
 							.UseSRGB(bUseSRGB)
 							.UseAlpha(bUseAlpha)
 					]
diff --git a/Engine/Source/Runtime/AppFramework/Public/Widgets/Colors/SColorThemes.h b/Engine/Source/Runtime/AppFramework/Public/Widgets/Colors/SColorThemes.h
--- a/Engine/Source/Runtime/AppFramework/Public/Widgets/Colors/SColorThemes.h
+++ b/Engine/Source/Runtime/AppFramework/Public/Widgets/Colors/SColorThemes.h
@@ -243,7 +243,7 @@ public:
 		SLATE_EVENT(FSimpleDelegate, HideTrashCallback)
 		
 		/** Specify what the bar should display when no colors are present */
-		SLATE_TEXT_ARGUMENT(EmptyText)
+		SLATE_ARGUMENT(FText, EmptyText)
 
 		/** Whether to display sRGB color */
 		SLATE_ATTRIBUTE(bool, UseSRGB)
