Changed SThemeColorBlock so that each block is split into non-alpha and alpha parts, so that colors with alpha of 0 can still be viewed.

diff --git a/Engine/Source/Runtime/AppFramework/Private/Widgets/Colors/SColorThemes.cpp b/Engine/Source/Runtime/AppFramework/Private/Widgets/Colors/SColorThemes.cpp
--- a/Engine/Source/Runtime/AppFramework/Private/Widgets/Colors/SColorThemes.cpp
+++ b/Engine/Source/Runtime/AppFramework/Private/Widgets/Colors/SColorThemes.cpp
@@ -296,12 +296,25 @@ void SThemeColorBlock::Construct(const FArguments& InArgs )
 			.Padding(FMargin(1.0f))
 			.ToolTip(ColorTooltip)
 			[
-				SNew(SColorBlock)
+				SNew(SHorizontalBox)
+				+ SHorizontalBox::Slot()
+				[
+					SNew(SColorBlock)
+					.Color(this, &SThemeColorBlock::GetColor)
+					.ColorIsHSV(true)
+					.IgnoreAlpha(true)
+					.ShowBackgroundForAlpha(false)
+					.UseSRGB(bUseSRGB)
+				]
+				+ SHorizontalBox::Slot()
+				[
+					SNew(SColorBlock)
 					.Color(this, &SThemeColorBlock::GetColor)
 					.ColorIsHSV(true)
 					.IgnoreAlpha(TAttribute<bool>::Create(TAttribute<bool>::FGetter::CreateSP(this, &SThemeColorBlock::OnReadIgnoreAlpha)))
 					.ShowBackgroundForAlpha(TAttribute<bool>::Create(TAttribute<bool>::FGetter::CreateSP(this, &SThemeColorBlock::OnReadShowBackgroundForAlpha)))
 					.UseSRGB(bUseSRGB)
+				]
 			]
 	];
 }
