Slate: enable mouse interaction with widgets in custom notifications.

diff --git a/Engine/Source/Runtime/AppFramework/Private/Widgets/Testing/STestSuite.cpp b/Engine/Source/Runtime/AppFramework/Private/Widgets/Testing/STestSuite.cpp
--- a/Engine/Source/Runtime/AppFramework/Private/Widgets/Testing/STestSuite.cpp
+++ b/Engine/Source/Runtime/AppFramework/Private/Widgets/Testing/STestSuite.cpp
@@ -3902,6 +3902,12 @@ public:
 						.Text(LOCTEXT("TestingBigTextBigMargin", "Big notififcation text!"))
 						.Font(FSlateFontInfo(FPaths::EngineContentDir() / TEXT("Slate/Fonts/Roboto-Bold.ttf"), 30))
 					]
+					+ SHorizontalBox::Slot()
+					.Padding(FMargin(15.0f, 0.0f, 0.0f, 0.0f))
+					[
+						SNew(SButton)
+						.Text(LOCTEXT("TestButtonInNotificaiton", "Button Test"))
+					]
 				]
 			];
 	}
diff --git a/Engine/Source/Runtime/Slate/Private/Widgets/Notifications/SNotificationList.cpp b/Engine/Source/Runtime/Slate/Private/Widgets/Notifications/SNotificationList.cpp
--- a/Engine/Source/Runtime/Slate/Private/Widgets/Notifications/SNotificationList.cpp
+++ b/Engine/Source/Runtime/Slate/Private/Widgets/Notifications/SNotificationList.cpp
@@ -642,6 +642,7 @@ public:
 						.Padding(0)
 						.BorderImage(FCoreStyle::Get().GetBrush("NotificationList.ItemBackground_Border_Transparent"))
 						.BorderBackgroundColor(this, &SNotificationItemExternalImpl::GetGlowColor)
+						.Visibility(EVisibility::SelfHitTestInvisible)
 					]
 				]
 			];
