Fixed issue where color picker eye dropper wasn't picking colors outside of a UE4 window.

diff --git a/Engine/Source/Runtime/AppFramework/Private/Widgets/Colors/SEyeDropperButton.h b/Engine/Source/Runtime/AppFramework/Private/Widgets/Colors/SEyeDropperButton.h
--- a/Engine/Source/Runtime/AppFramework/Private/Widgets/Colors/SEyeDropperButton.h
+++ b/Engine/Source/Runtime/AppFramework/Private/Widgets/Colors/SEyeDropperButton.h
@@ -75,6 +75,29 @@ public:
 		);
 	}
 
+	virtual void Tick(const FGeometry& AllottedGeometry, const double InCurrentTime, const float InDeltaTime) override
+	{
+		if (bWasClickActivated && HasMouseCapture() && bWasLeft && !bWasReEntered)
+		{
+			FVector2D CurrentCursorPosition = FSlateApplication::Get().GetCursorPos();
+			if (CurrentCursorPosition != LastCursorPosition)
+			{
+				// In dropper mode and outside the button - sample the pixel color and push it to the client
+				FLinearColor ScreenColor = FPlatformMisc::GetScreenPixelColor(CurrentCursorPosition, false);
+				OnValueChanged.ExecuteIfBound(ScreenColor);
+			}
+
+			LastCursorPosition = CurrentCursorPosition;
+		}
+		else
+		{
+			LastCursorPosition.Reset();
+		}
+
+		SButton::Tick(AllottedGeometry, InCurrentTime, InDeltaTime);
+	}
+
+
 	virtual FReply OnMouseButtonDown( const FGeometry& MyGeometry, const FPointerEvent& MouseEvent ) override
 	{
 		// Clicking ANY mouse button when the dropper isn't active resets the active dropper states ready to activate
@@ -150,10 +173,6 @@ public:
 				// Mouse is outside the button
 				bWasLeft = true;
 				bWasReEntered = false;
-
-				// In dropper mode and outside the button - sample the pixel color and push it to the client
-				FLinearColor ScreenColor = FPlatformMisc::GetScreenPixelColor(FSlateApplication::Get().GetCursorPos(), false);
-				OnValueChanged.ExecuteIfBound(ScreenColor);
 			}
 		}
 
@@ -243,6 +262,9 @@ private:
 	/** Sets the display Gamma setting - used to correct colors sampled from the screen */
 	TAttribute<float> DisplayGamma;
 
+	/** Previous Tick's cursor position */
+	TOptional<FVector2D> LastCursorPosition;
+
 	// Dropper states
 	bool bWasClicked;
 	bool bWasClickActivated;
