Slate - ESlateCheckBoxState has been renamed to ECheckBoxState.

diff --git a/Engine/Source/Editor/UnrealEd/Public/STransformViewportToolbar.h b/Engine/Source/Editor/UnrealEd/Public/STransformViewportToolbar.h
--- a/Engine/Source/Editor/UnrealEd/Public/STransformViewportToolbar.h
+++ b/Engine/Source/Editor/UnrealEd/Public/STransformViewportToolbar.h
@@ -78,14 +78,14 @@ private:
 	TSharedRef< SWidget > MakeSurfaceSnappingButton( FName ToolBarStyle );
 
 	/** Grid Snap checked state callbacks */
-	ESlateCheckBoxState::Type IsLocationGridSnapChecked() const;
-	ESlateCheckBoxState::Type IsRotationGridSnapChecked() const;
-	ESlateCheckBoxState::Type IsScaleGridSnapChecked() const;
+	ECheckBoxState IsLocationGridSnapChecked() const;
+	ECheckBoxState IsRotationGridSnapChecked() const;
+	ECheckBoxState IsScaleGridSnapChecked() const;
 
 	/** Grid snap toggle handlers */
-	void HandleToggleLocationGridSnap(ESlateCheckBoxState::Type InState);
-	void HandleToggleRotationGridSnap(ESlateCheckBoxState::Type InState);
-	void HandleToggleScaleGridSnap(ESlateCheckBoxState::Type InState);
+	void HandleToggleLocationGridSnap(ECheckBoxState InState);
+	void HandleToggleRotationGridSnap(ECheckBoxState InState);
+	void HandleToggleScaleGridSnap(ECheckBoxState InState);
 
 private:
 
diff --git a/Engine/Source/Editor/UnrealEd/Public/SViewportToolBarComboMenu.h b/Engine/Source/Editor/UnrealEd/Public/SViewportToolBarComboMenu.h
--- a/Engine/Source/Editor/UnrealEd/Public/SViewportToolBarComboMenu.h
+++ b/Engine/Source/Editor/UnrealEd/Public/SViewportToolBarComboMenu.h
@@ -26,7 +26,7 @@ public:
 		SLATE_EVENT( FOnCheckStateChanged, OnCheckStateChanged )
 			
 		/** Sets the current checked state of the checkbox */
-		SLATE_ATTRIBUTE( ESlateCheckBoxState::Type, IsChecked )
+		SLATE_ATTRIBUTE( ECheckBoxState, IsChecked )
 
 		/** Icon shown on the toggle button */
 		SLATE_ATTRIBUTE( FSlateIcon, Icon )
diff --git a/Engine/Source/Editor/WorldBrowser/Private/Tiles/STiledLandscapeImportDlg.cpp b/Engine/Source/Editor/WorldBrowser/Private/Tiles/STiledLandscapeImportDlg.cpp
--- a/Engine/Source/Editor/WorldBrowser/Private/Tiles/STiledLandscapeImportDlg.cpp
+++ b/Engine/Source/Editor/WorldBrowser/Private/Tiles/STiledLandscapeImportDlg.cpp
@@ -688,14 +688,14 @@ FText STiledLandcapeImportDlg::GetWeightmapCountText(TSharedPtr<FTiledLandscapeI
 	return FText::AsNumber(NumWeighmaps);
 }
 
-ESlateCheckBoxState::Type STiledLandcapeImportDlg::GetLayerBlendState(TSharedPtr<FTiledLandscapeImportSettings::LandscapeLayerSettings> InLayerData) const
+ECheckBoxState STiledLandcapeImportDlg::GetLayerBlendState(TSharedPtr<FTiledLandscapeImportSettings::LandscapeLayerSettings> InLayerData) const
 {
-	return (InLayerData->bNoBlendWeight ? ESlateCheckBoxState::Unchecked : ESlateCheckBoxState::Checked);
+	return (InLayerData->bNoBlendWeight ? ECheckBoxState::Unchecked : ECheckBoxState::Checked);
 }
 
-void STiledLandcapeImportDlg::OnLayerBlendStateChanged(ESlateCheckBoxState::Type NewState, TSharedPtr<FTiledLandscapeImportSettings::LandscapeLayerSettings> InLayerData)
+void STiledLandcapeImportDlg::OnLayerBlendStateChanged(ECheckBoxState NewState, TSharedPtr<FTiledLandscapeImportSettings::LandscapeLayerSettings> InLayerData)
 {
-	InLayerData->bNoBlendWeight = !(NewState == ESlateCheckBoxState::Checked);
+	InLayerData->bNoBlendWeight = !(NewState == ECheckBoxState::Checked);
 }
 
 FText STiledLandcapeImportDlg::GenerateConfigurationText(int32 NumComponents, int32 NumSectionsPerComponent, int32 NumQuadsPerSection) const
diff --git a/Engine/Source/Editor/WorldBrowser/Private/Tiles/STiledLandscapeImportDlg.h b/Engine/Source/Editor/WorldBrowser/Private/Tiles/STiledLandscapeImportDlg.h
--- a/Engine/Source/Editor/WorldBrowser/Private/Tiles/STiledLandscapeImportDlg.h
+++ b/Engine/Source/Editor/WorldBrowser/Private/Tiles/STiledLandscapeImportDlg.h
@@ -73,8 +73,8 @@ private:
 	FText GetWeightmapCountText(TSharedPtr<FTiledLandscapeImportSettings::LandscapeLayerSettings> InLayerData) const; 
 
 	/** */
-	ESlateCheckBoxState::Type GetLayerBlendState(TSharedPtr<FTiledLandscapeImportSettings::LandscapeLayerSettings> InLayerData) const;
-	void OnLayerBlendStateChanged(ESlateCheckBoxState::Type NewState, TSharedPtr<FTiledLandscapeImportSettings::LandscapeLayerSettings> InLayerData);
+	ECheckBoxState GetLayerBlendState(TSharedPtr<FTiledLandscapeImportSettings::LandscapeLayerSettings> InLayerData) const;
+	void OnLayerBlendStateChanged(ECheckBoxState NewState, TSharedPtr<FTiledLandscapeImportSettings::LandscapeLayerSettings> InLayerData);
 
 	/** */
 	int32 SetPossibleConfigurationsForFileSize(int64 InFileSize);
diff --git a/Engine/Source/Editor/WorldBrowser/Private/Tiles/SWorldLayers.cpp b/Engine/Source/Editor/WorldBrowser/Private/Tiles/SWorldLayers.cpp
--- a/Engine/Source/Editor/WorldBrowser/Private/Tiles/SWorldLayers.cpp
+++ b/Engine/Source/Editor/WorldBrowser/Private/Tiles/SWorldLayers.cpp
@@ -207,9 +207,9 @@ void SWorldLayerButton::Construct(const FArguments& InArgs)
 	CheckBox->SetOnLayerDoubleClicked(FOnClicked::CreateSP(this, &SWorldLayerButton::OnDoubleClicked));
 }
 
-void SWorldLayerButton::OnCheckStateChanged(ESlateCheckBoxState::Type NewState)
+void SWorldLayerButton::OnCheckStateChanged(ECheckBoxState NewState)
 {
-	if (NewState == ESlateCheckBoxState::Checked)
+	if (NewState == ECheckBoxState::Checked)
 	{
 		WorldModel->SetSelectedLayer(WorldLayer);
 	}
@@ -219,9 +219,9 @@ void SWorldLayerButton::OnCheckStateChanged(ESlateCheckBoxState::Type NewState)
 	}
 }
 
-ESlateCheckBoxState::Type SWorldLayerButton::IsChecked() const
+ECheckBoxState SWorldLayerButton::IsChecked() const
 {
-	return WorldModel->IsLayerSelected(WorldLayer) ? ESlateCheckBoxState::Checked : ESlateCheckBoxState::Unchecked;
+	return WorldModel->IsLayerSelected(WorldLayer) ? ECheckBoxState::Checked : ECheckBoxState::Unchecked;
 }
 
 /** Handler for when the filter checkbox is double clicked */
diff --git a/Engine/Source/Editor/WorldBrowser/Private/Tiles/SWorldLayers.h b/Engine/Source/Editor/WorldBrowser/Private/Tiles/SWorldLayers.h
--- a/Engine/Source/Editor/WorldBrowser/Private/Tiles/SWorldLayers.h
+++ b/Engine/Source/Editor/WorldBrowser/Private/Tiles/SWorldLayers.h
@@ -20,8 +20,8 @@ public:
 
 	~SWorldLayerButton();
 	void Construct(const FArguments& InArgs);
-	void OnCheckStateChanged(ESlateCheckBoxState::Type NewState);
-	ESlateCheckBoxState::Type IsChecked() const;
+	void OnCheckStateChanged(ECheckBoxState NewState);
+	ECheckBoxState IsChecked() const;
 	FReply OnDoubleClicked();
 	FReply OnCtrlClicked();
 	TSharedRef<SWidget> GetRightClickMenu();
@@ -52,9 +52,9 @@ public:
 		return LayerData.StreamingDistance;
 	}
 	
-	ESlateCheckBoxState::Type GetDistanceStreamingState() const
+	ECheckBoxState GetDistanceStreamingState() const
 	{
-		return IsDistanceStreamingEnabled() ? ESlateCheckBoxState::Checked : ESlateCheckBoxState::Unchecked;
+		return IsDistanceStreamingEnabled() ? ECheckBoxState::Checked : ECheckBoxState::Unchecked;
 	}
 
 	bool IsDistanceStreamingEnabled() const
@@ -62,9 +62,9 @@ public:
 		return LayerData.DistanceStreamingEnabled;
 	}
 	
-	void OnDistanceStreamingStateChanged(ESlateCheckBoxState::Type NewState)
+	void OnDistanceStreamingStateChanged(ECheckBoxState NewState)
 	{
-		SetDistanceStreamingState(NewState == ESlateCheckBoxState::Checked);
+		SetDistanceStreamingState(NewState == ECheckBoxState::Checked);
 	}
 
 	FText GetLayerName() const
diff --git a/Engine/Source/Programs/SymbolDebugger/Private/SSymbolDebugger.cpp b/Engine/Source/Programs/SymbolDebugger/Private/SSymbolDebugger.cpp
--- a/Engine/Source/Programs/SymbolDebugger/Private/SSymbolDebugger.cpp
+++ b/Engine/Source/Programs/SymbolDebugger/Private/SSymbolDebugger.cpp
@@ -339,19 +339,19 @@ FReply SSymbolDebugger::FileOpenClicked()
 	return FReply::Unhandled();
 }
 
-ESlateCheckBoxState::Type SSymbolDebugger::IsMethodChecked(ESymbolDebuggerMethods InMethod) const
+ECheckBoxState SSymbolDebugger::IsMethodChecked(ESymbolDebuggerMethods InMethod) const
 {
 	if (GetSymbolDebuggertMethod() == InMethod)
 	{
-		return ESlateCheckBoxState::Checked;
+		return ECheckBoxState::Checked;
 	}
 
-	return ESlateCheckBoxState::Unchecked;
+	return ECheckBoxState::Unchecked;
 }
 
-void SSymbolDebugger::OnMethodChanged(ESlateCheckBoxState::Type InNewRadioState, ESymbolDebuggerMethods InMethodThatChanged)
+void SSymbolDebugger::OnMethodChanged(ECheckBoxState InNewRadioState, ESymbolDebuggerMethods InMethodThatChanged)
 {
-	if (InNewRadioState == ESlateCheckBoxState::Checked)
+	if (InNewRadioState == ECheckBoxState::Checked)
 	{
 		Delegate_OnSetCurrentMethod.Execute(InMethodThatChanged);
 	}
diff --git a/Engine/Source/Programs/SymbolDebugger/Private/SSymbolDebugger.h b/Engine/Source/Programs/SymbolDebugger/Private/SSymbolDebugger.h
--- a/Engine/Source/Programs/SymbolDebugger/Private/SSymbolDebugger.h
+++ b/Engine/Source/Programs/SymbolDebugger/Private/SSymbolDebugger.h
@@ -208,9 +208,9 @@ public:
 	 *
 	 *	@param	InMethod				The debugger method to check
 	 *
-	 *	@return	ESlateCheckBoxState		The state of the button for the given method
+	 *	@return	ECheckBoxState			The state of the button for the given method
 	 */
-	ESlateCheckBoxState::Type IsMethodChecked(ESymbolDebuggerMethods InMethod) const;
+	ECheckBoxState IsMethodChecked(ESymbolDebuggerMethods InMethod) const;
 
 	/**
 	 *	Handle the Method radio buttons changing
@@ -218,7 +218,7 @@ public:
 	 *	@param	InNewRadioState			The new state of the radio button
 	 *	@param	InMethodThatChanged		The method radio button that changed
 	 */
-	void OnMethodChanged(ESlateCheckBoxState::Type InNewRadioState, ESymbolDebuggerMethods InMethodThatChanged);
+	void OnMethodChanged(ECheckBoxState InNewRadioState, ESymbolDebuggerMethods InMethodThatChanged);
 
 	/**
 	 *	Is the textbox enabled for the given TextField?
diff --git a/Engine/Source/Programs/UnrealSync/Private/GUI.cpp b/Engine/Source/Programs/UnrealSync/Private/GUI.cpp
--- a/Engine/Source/Programs/UnrealSync/Private/GUI.cpp
+++ b/Engine/Source/Programs/UnrealSync/Private/GUI.cpp
@@ -236,11 +236,11 @@ public:
 		 *
 		 * @returns State of the check box enum.
 		 */
-		ESlateCheckBoxState::Type IsCheckboxChecked() const
+		ECheckBoxState IsCheckboxChecked() const
 		{
 			return Parent.GetChosen() == Id
-				? ESlateCheckBoxState::Checked
-				: ESlateCheckBoxState::Unchecked;
+				? ECheckBoxState::Checked
+				: ECheckBoxState::Unchecked;
 		}
 
 	private:
@@ -249,7 +249,7 @@ public:
 		 *
 		 * @parent InNewState New state enum.
 		 */
-		void OnCheckStateChange(ESlateCheckBoxState::Type InNewState)
+		void OnCheckStateChange(ECheckBoxState InNewState)
 		{
 			Parent.ChooseEnabledItem(Id);
 		}
diff --git a/Engine/Source/Runtime/AppFramework/Private/Framework/Testing/SWidgetGallery.cpp b/Engine/Source/Runtime/AppFramework/Private/Framework/Testing/SWidgetGallery.cpp
--- a/Engine/Source/Runtime/AppFramework/Private/Framework/Testing/SWidgetGallery.cpp
+++ b/Engine/Source/Runtime/AppFramework/Private/Framework/Testing/SWidgetGallery.cpp
@@ -943,12 +943,12 @@ private:
 
 	EPopupMethod PopupMethod;
 
-	ESlateCheckBoxState::Type IsCreateNewWindowChecked() const
+	ECheckBoxState IsCreateNewWindowChecked() const
 	{
-		return this->PopupMethod == EPopupMethod::CreateNewWindow ? ESlateCheckBoxState::Checked : ESlateCheckBoxState::Unchecked;
+		return this->PopupMethod == EPopupMethod::CreateNewWindow ? ECheckBoxState::Checked : ECheckBoxState::Unchecked;
 	}
 
-	void CreateNewWindowToggled(ESlateCheckBoxState::Type NewState)
+	void CreateNewWindowToggled(ECheckBoxState NewState)
 	{
 		this->PopupMethod = (PopupMethod == EPopupMethod::CreateNewWindow)
 			? EPopupMethod::UseCurrentWindow
@@ -965,17 +965,17 @@ private:
 	}
 
 	// Callback for changing the checked state of a check box.
-	void HandleCheckBoxCheckedStateChanged( ESlateCheckBoxState::Type NewState, bool* CheckBoxThatChanged )
+	void HandleCheckBoxCheckedStateChanged( ECheckBoxState NewState, bool* CheckBoxThatChanged )
 	{
-		*CheckBoxThatChanged = (NewState == ESlateCheckBoxState::Checked);
+		*CheckBoxThatChanged = (NewState == ECheckBoxState::Checked);
 	}
 
 	// Callback for determining whether a check box is checked.
-	ESlateCheckBoxState::Type HandleCheckBoxIsChecked( bool* CheckBox ) const
+	ECheckBoxState HandleCheckBoxIsChecked( bool* CheckBox ) const
 	{
 		return (*CheckBox)
-			? ESlateCheckBoxState::Checked
-			: ESlateCheckBoxState::Unchecked;
+			? ECheckBoxState::Checked
+			: ECheckBoxState::Unchecked;
 	}
 
 	// Callback for generating a widget in the SComboBox example.
@@ -998,20 +998,20 @@ private:
 	}
 
 	// Callback for checking a radio button.
-	void HandleRadioButtonCheckStateChanged( ESlateCheckBoxState::Type NewRadioState, ERadioChoice RadioThatChanged )
+	void HandleRadioButtonCheckStateChanged( ECheckBoxState NewRadioState, ERadioChoice RadioThatChanged )
 	{
-		if (NewRadioState == ESlateCheckBoxState::Checked)
+		if (NewRadioState == ECheckBoxState::Checked)
 		{
 			RadioChoice = RadioThatChanged;
 		}
 	}
 
 	// Callback for determining whether a radio button is checked.
-	ESlateCheckBoxState::Type HandleRadioButtonIsChecked( ERadioChoice ButtonId ) const
+	ECheckBoxState HandleRadioButtonIsChecked( ERadioChoice ButtonId ) const
 	{
 		return (RadioChoice == ButtonId)
-			? ESlateCheckBoxState::Checked
-			: ESlateCheckBoxState::Unchecked;
+			? ECheckBoxState::Checked
+			: ECheckBoxState::Unchecked;
 	}
 
 	// Callback for searching in the SSearchBox example.
diff --git a/Engine/Source/Runtime/AppFramework/Private/Widgets/Colors/SColorPicker.cpp b/Engine/Source/Runtime/AppFramework/Private/Widgets/Colors/SColorPicker.cpp
--- a/Engine/Source/Runtime/AppFramework/Private/Widgets/Colors/SColorPicker.cpp
+++ b/Engine/Source/Runtime/AppFramework/Private/Widgets/Colors/SColorPicker.cpp
@@ -1419,9 +1419,9 @@ void SColorPicker::HandleRGBColorChanged( FLinearColor NewValue )
 }
 
 
-void SColorPicker::HandleSRGBCheckBoxCheckStateChanged( ESlateCheckBoxState::Type InIsChecked )
+void SColorPicker::HandleSRGBCheckBoxCheckStateChanged( ECheckBoxState InIsChecked )
 {
-	SColorThemesViewer::bSRGBEnabled = (InIsChecked == ESlateCheckBoxState::Checked);
+	SColorThemesViewer::bSRGBEnabled = (InIsChecked == ECheckBoxState::Checked);
 
 	if (FPaths::FileExists(GEditorUserSettingsIni))
 	{
@@ -1430,9 +1430,9 @@ void SColorPicker::HandleSRGBCheckBoxCheckStateChanged( ESlateCheckBoxState::Typ
 }
 
 
-ESlateCheckBoxState::Type SColorPicker::HandleSRGBCheckBoxIsChecked() const
+ECheckBoxState SColorPicker::HandleSRGBCheckBoxIsChecked() const
 {
-	return SColorThemesViewer::bSRGBEnabled ? ESlateCheckBoxState::Checked : ESlateCheckBoxState::Unchecked;
+	return SColorThemesViewer::bSRGBEnabled ? ECheckBoxState::Checked : ECheckBoxState::Unchecked;
 }
 
 
diff --git a/Engine/Source/Runtime/AppFramework/Private/Widgets/Testing/STestSuite.cpp b/Engine/Source/Runtime/AppFramework/Private/Widgets/Testing/STestSuite.cpp
--- a/Engine/Source/Runtime/AppFramework/Private/Widgets/Testing/STestSuite.cpp
+++ b/Engine/Source/Runtime/AppFramework/Private/Widgets/Testing/STestSuite.cpp
@@ -1175,14 +1175,14 @@ private:
 		return bIsReadOnly;
 	}
 
-	ESlateCheckBoxState::Type IsReadOnlyChecked() const
+	ECheckBoxState IsReadOnlyChecked() const
 	{
-		return (bIsReadOnly) ? ESlateCheckBoxState::Checked : ESlateCheckBoxState::Unchecked;
+		return (bIsReadOnly) ? ECheckBoxState::Checked : ECheckBoxState::Unchecked;
 	}
 
-	void OnReadOnlyCheckedStateChanged(ESlateCheckBoxState::Type InState)
+	void OnReadOnlyCheckedStateChanged(ECheckBoxState InState)
 	{
-		bIsReadOnly = (InState == ESlateCheckBoxState::Checked);
+		bIsReadOnly = (InState == ECheckBoxState::Checked);
 	}
 
 	bool bIsReadOnly;
@@ -1882,14 +1882,14 @@ public:
 		WrapWidth = Value;
 	}
 
-	ESlateCheckBoxState::Type ShouldWrapRichText() const
+	ECheckBoxState ShouldWrapRichText() const
 	{
-		return bShouldWrap ? ESlateCheckBoxState::Checked : ESlateCheckBoxState::Unchecked;
+		return bShouldWrap ? ECheckBoxState::Checked : ECheckBoxState::Unchecked;
 	}
 
-	void ShouldWrapRichTextChanged( ESlateCheckBoxState::Type CheckState )
+	void ShouldWrapRichTextChanged( ECheckBoxState CheckState )
 	{
-		bShouldWrap = CheckState == ESlateCheckBoxState::Checked ? true : false;
+		bShouldWrap = CheckState == ECheckBoxState::Checked ? true : false;
 	}
 
 	float GetRichTextWrapWidthValue() const
@@ -2796,14 +2796,14 @@ protected:
 		StyleSelectedText();
 	}
 
-	ESlateCheckBoxState::Type IsFontStyleBold() const
+	ECheckBoxState IsFontStyleBold() const
 	{
-		return (FontStyle & FTextStyles::EFontStyle::Bold) ? ESlateCheckBoxState::Checked : ESlateCheckBoxState::Unchecked;
+		return (FontStyle & FTextStyles::EFontStyle::Bold) ? ECheckBoxState::Checked : ECheckBoxState::Unchecked;
 	}
 
-	void OnFontStyleBoldChanged(ESlateCheckBoxState::Type InState)
+	void OnFontStyleBoldChanged(ECheckBoxState InState)
 	{
-		if(InState == ESlateCheckBoxState::Checked)
+		if(InState == ECheckBoxState::Checked)
 		{
 			FontStyle |= FTextStyles::EFontStyle::Bold;
 		}
@@ -2814,14 +2814,14 @@ protected:
 		StyleSelectedText();
 	}
 
-	ESlateCheckBoxState::Type IsFontStyleItalic() const
+	ECheckBoxState IsFontStyleItalic() const
 	{
-		return (FontStyle & FTextStyles::EFontStyle::Italic) ? ESlateCheckBoxState::Checked : ESlateCheckBoxState::Unchecked;
+		return (FontStyle & FTextStyles::EFontStyle::Italic) ? ECheckBoxState::Checked : ECheckBoxState::Unchecked;
 	}
 
-	void OnFontStyleItalicChanged(ESlateCheckBoxState::Type InState)
+	void OnFontStyleItalicChanged(ECheckBoxState InState)
 	{
-		if(InState == ESlateCheckBoxState::Checked)
+		if(InState == ECheckBoxState::Checked)
 		{
 			FontStyle |= FTextStyles::EFontStyle::Italic;
 		}
@@ -2924,14 +2924,14 @@ protected:
 		return FReply::Handled();
 	}
 
-	ESlateCheckBoxState::Type IsEnableSyntaxHighlightingChecked() const
+	ECheckBoxState IsEnableSyntaxHighlightingChecked() const
 	{
-		return (SyntaxHighlighterMarshaller->IsSyntaxHighlightingEnabled()) ? ESlateCheckBoxState::Checked : ESlateCheckBoxState::Unchecked;
+		return (SyntaxHighlighterMarshaller->IsSyntaxHighlightingEnabled()) ? ECheckBoxState::Checked : ECheckBoxState::Unchecked;
 	}
 
-	void OnEnableSyntaxHighlightingChanged(ESlateCheckBoxState::Type InState)
+	void OnEnableSyntaxHighlightingChanged(ECheckBoxState InState)
 	{
-		SyntaxHighlighterMarshaller->EnableSyntaxHighlighting(InState == ESlateCheckBoxState::Checked);
+		SyntaxHighlighterMarshaller->EnableSyntaxHighlighting(InState == ECheckBoxState::Checked);
 	}
 
 protected:
@@ -3506,7 +3506,7 @@ class SFxTest : public SCompoundWidget
 		RenderScaleOrigin = FVector2D(0.5f, 0.5f);
 		LayoutScale = 1.0f;
 		VisualOffset = FVector2D::ZeroVector;
-		FxWidgetIgnoreClippingState = ESlateCheckBoxState::Checked;
+		FxWidgetIgnoreClippingState = ECheckBoxState::Checked;
 
 
 		this->ChildSlot
@@ -3659,10 +3659,10 @@ class SFxTest : public SCompoundWidget
 	}
 	END_SLATE_FUNCTION_BUILD_OPTIMIZATION
 
-	bool GetFxWidgetIgnoreClipping() const { return FxWidgetIgnoreClippingState == ESlateCheckBoxState::Checked; }
-	ESlateCheckBoxState::Type GetFxWidgetIgnoreClippingState() const { return FxWidgetIgnoreClippingState; }
-	void OnFxWidgetIgnoreClippingChanged( ESlateCheckBoxState::Type InValue ) { FxWidgetIgnoreClippingState = InValue; }
-	ESlateCheckBoxState::Type FxWidgetIgnoreClippingState;
+	bool GetFxWidgetIgnoreClipping() const { return FxWidgetIgnoreClippingState == ECheckBoxState::Checked; }
+	ECheckBoxState GetFxWidgetIgnoreClippingState() const { return FxWidgetIgnoreClippingState; }
+	void OnFxWidgetIgnoreClippingChanged( ECheckBoxState InValue ) { FxWidgetIgnoreClippingState = InValue; }
+	ECheckBoxState FxWidgetIgnoreClippingState;
 
 	float GetRenderScale() const { return RenderScale; }
 	void OnRenderScaleChanged( float InValue ) { RenderScale = InValue; }
@@ -3754,18 +3754,18 @@ class SDPIScalingTest : public SCompoundWidget
 		DPIScale = InScale;
 	}
 
-	ESlateCheckBoxState::Type IsFillChecked() const
+	ECheckBoxState IsFillChecked() const
 	{
 		const bool bIsFilling = (ScalerSlot->HAlignment == HAlign_Fill);
 		return (bIsFilling)
-			? ESlateCheckBoxState::Checked
-			: ESlateCheckBoxState::Unchecked;
+			? ECheckBoxState::Checked
+			: ECheckBoxState::Unchecked;
 	}
 
-	void OnFillChecked(ESlateCheckBoxState::Type InValue)
+	void OnFillChecked(ECheckBoxState InValue)
 	{
-		ScalerSlot->HAlign( (InValue == ESlateCheckBoxState::Checked) ? HAlign_Fill : HAlign_Center );
-		ScalerSlot->VAlign( (InValue == ESlateCheckBoxState::Checked) ? VAlign_Fill : VAlign_Center );
+		ScalerSlot->HAlign( (InValue == ECheckBoxState::Checked) ? HAlign_Fill : HAlign_Center );
+		ScalerSlot->VAlign( (InValue == ECheckBoxState::Checked) ? VAlign_Fill : VAlign_Center );
 	}
 
 	float DPIScale;
@@ -3935,7 +3935,7 @@ public:
 		bAddDummyCheckBox = false;
 		bAddDummyHyperlink = false;
 
-		DummyCheckBoxState = ESlateCheckBoxState::Unchecked;
+		DummyCheckBoxState = ECheckBoxState::Unchecked;
 
 		this->ChildSlot
 		[
@@ -4051,44 +4051,44 @@ public:
 	}
 
 protected:
-	ESlateCheckBoxState::Type IsUseLargeFontChecked() const
+	ECheckBoxState IsUseLargeFontChecked() const
 	{
-		return bUseLargeFont ? ESlateCheckBoxState::Checked : ESlateCheckBoxState::Unchecked;
+		return bUseLargeFont ? ECheckBoxState::Checked : ECheckBoxState::Unchecked;
 	}
 
-	void OnUseLargeFontCheckStateChanged(ESlateCheckBoxState::Type NewState)
+	void OnUseLargeFontCheckStateChanged(ECheckBoxState NewState)
 	{
-		bUseLargeFont = (NewState == ESlateCheckBoxState::Checked);
+		bUseLargeFont = (NewState == ECheckBoxState::Checked);
 	}
 
-	ESlateCheckBoxState::Type IsAddDummyButtonsChecked() const
+	ECheckBoxState IsAddDummyButtonsChecked() const
 	{
-		return bAddDummyButtons ? ESlateCheckBoxState::Checked : ESlateCheckBoxState::Unchecked;
+		return bAddDummyButtons ? ECheckBoxState::Checked : ECheckBoxState::Unchecked;
 	}
 
-	void OnAddDummyButtonsCheckStateChanged(ESlateCheckBoxState::Type NewState)
+	void OnAddDummyButtonsCheckStateChanged(ECheckBoxState NewState)
 	{
-		bAddDummyButtons = (NewState == ESlateCheckBoxState::Checked);
+		bAddDummyButtons = (NewState == ECheckBoxState::Checked);
 	}
 
-	ESlateCheckBoxState::Type IsAddDummyCheckBoxChecked() const
+	ECheckBoxState IsAddDummyCheckBoxChecked() const
 	{
-		return bAddDummyCheckBox ? ESlateCheckBoxState::Checked : ESlateCheckBoxState::Unchecked;
+		return bAddDummyCheckBox ? ECheckBoxState::Checked : ECheckBoxState::Unchecked;
 	}
 
-	void OnAddDummyCheckBoxCheckStateChanged(ESlateCheckBoxState::Type NewState)
+	void OnAddDummyCheckBoxCheckStateChanged(ECheckBoxState NewState)
 	{
-		bAddDummyCheckBox = (NewState == ESlateCheckBoxState::Checked);
+		bAddDummyCheckBox = (NewState == ECheckBoxState::Checked);
 	}
 
-	ESlateCheckBoxState::Type IsAddDummyHyperlinkChecked() const
+	ECheckBoxState IsAddDummyHyperlinkChecked() const
 	{
-		return bAddDummyHyperlink ? ESlateCheckBoxState::Checked : ESlateCheckBoxState::Unchecked;
+		return bAddDummyHyperlink ? ECheckBoxState::Checked : ECheckBoxState::Unchecked;
 	}
 
-	void OnAddDummyHyperlinkCheckStateChanged(ESlateCheckBoxState::Type NewState)
+	void OnAddDummyHyperlinkCheckStateChanged(ECheckBoxState NewState)
 	{
-		bAddDummyHyperlink = (NewState == ESlateCheckBoxState::Checked);
+		bAddDummyHyperlink = (NewState == ECheckBoxState::Checked);
 	}
 
 	void SetNotificationInfoFlags(FNotificationInfo& Info)
@@ -4121,7 +4121,7 @@ protected:
 
 		if(bAddDummyCheckBox)
 		{
-			Info.CheckBoxState = TAttribute<ESlateCheckBoxState::Type>(this, &SNotificationListTest::GetDummyCheckBoxState);
+			Info.CheckBoxState = TAttribute<ECheckBoxState>(this, &SNotificationListTest::GetDummyCheckBoxState);
 			Info.CheckBoxStateChanged = FOnCheckStateChanged::CreateSP(this, &SNotificationListTest::OnDummyCheckBoxStateChanged);
 			Info.CheckBoxText = LOCTEXT("NotificationListTest-DummyCheckBoxText", "Dummy Check Box");
 		}
@@ -4133,12 +4133,12 @@ protected:
 		}
 	}
 
-	ESlateCheckBoxState::Type GetDummyCheckBoxState() const
+	ECheckBoxState GetDummyCheckBoxState() const
 	{
 		return DummyCheckBoxState;
 	}
 
-	void OnDummyCheckBoxStateChanged(ESlateCheckBoxState::Type NewState)
+	void OnDummyCheckBoxStateChanged(ECheckBoxState NewState)
 	{
 		DummyCheckBoxState = NewState;
 	}
@@ -4238,7 +4238,7 @@ protected:
 	/** The pending progress message */
 	TWeakPtr<SNotificationItem> PendingProgressPtr;
 	/** If showing a dummy check box, this maintains its state in the UI */
-	ESlateCheckBoxState::Type DummyCheckBoxState;
+	ECheckBoxState DummyCheckBoxState;
 };
 
 class SGridPanelTest : public SCompoundWidget
diff --git a/Engine/Source/Runtime/AppFramework/Private/Widgets/Workflow/SWizard.cpp b/Engine/Source/Runtime/AppFramework/Private/Widgets/Workflow/SWizard.cpp
--- a/Engine/Source/Runtime/AppFramework/Private/Widgets/Workflow/SWizard.cpp
+++ b/Engine/Source/Runtime/AppFramework/Private/Widgets/Workflow/SWizard.cpp
@@ -287,23 +287,23 @@ EVisibility SWizard::HandleNextButtonVisibility() const
 }
 
 
-void SWizard::HandlePageButtonCheckStateChanged( ESlateCheckBoxState::Type NewState, int32 PageIndex )
+void SWizard::HandlePageButtonCheckStateChanged( ECheckBoxState NewState, int32 PageIndex )
 {
-	if (NewState == ESlateCheckBoxState::Checked)
+	if (NewState == ECheckBoxState::Checked)
 	{
 		ShowPage(PageIndex);
 	}
 }
 
 
-ESlateCheckBoxState::Type SWizard::HandlePageButtonIsChecked( int32 PageIndex ) const
+ECheckBoxState SWizard::HandlePageButtonIsChecked( int32 PageIndex ) const
 {
 	if (PageIndex == WidgetSwitcher->GetActiveWidgetIndex())
 	{
-		return ESlateCheckBoxState::Checked;
+		return ECheckBoxState::Checked;
 	}
 
-	return ESlateCheckBoxState::Unchecked;
+	return ECheckBoxState::Unchecked;
 }
 
 
diff --git a/Engine/Source/Runtime/AppFramework/Public/Widgets/Colors/SColorPicker.h b/Engine/Source/Runtime/AppFramework/Public/Widgets/Colors/SColorPicker.h
--- a/Engine/Source/Runtime/AppFramework/Public/Widgets/Colors/SColorPicker.h
+++ b/Engine/Source/Runtime/AppFramework/Public/Widgets/Colors/SColorPicker.h
@@ -307,10 +307,10 @@ private:
 	void HandleRGBColorChanged( FLinearColor NewValue );
 
 	// Callback for changing the checked state of the sRGB check box.
-	void HandleSRGBCheckBoxCheckStateChanged( ESlateCheckBoxState::Type InIsChecked );
+	void HandleSRGBCheckBoxCheckStateChanged( ECheckBoxState InIsChecked );
 
 	// Callback for determining whether the sRGB check box should be checked.
-	ESlateCheckBoxState::Type HandleSRGBCheckBoxIsChecked() const;
+	ECheckBoxState HandleSRGBCheckBoxIsChecked() const;
 
 	// Callback for selecting a color in the color theme bar.
 	void HandleThemeBarColorSelected( FLinearColor NewValue );
diff --git a/Engine/Source/Runtime/AppFramework/Public/Widgets/Views/SMultipleOptionTable.h b/Engine/Source/Runtime/AppFramework/Public/Widgets/Views/SMultipleOptionTable.h
--- a/Engine/Source/Runtime/AppFramework/Public/Widgets/Views/SMultipleOptionTable.h
+++ b/Engine/Source/Runtime/AppFramework/Public/Widgets/Views/SMultipleOptionTable.h
@@ -22,7 +22,7 @@ public:
 		SLATE_ARGUMENT(TSharedPtr<STableViewBase>, OwnerTableView)
 		SLATE_ARGUMENT(OptionType, Option)
 		SLATE_EVENT(FOnCheckStateChanged, OnCheckStateChanged)
-		SLATE_ATTRIBUTE(ESlateCheckBoxState::Type, IsChecked)
+		SLATE_ATTRIBUTE(ECheckBoxState, IsChecked)
 		SLATE_EVENT(FOnGenerateWidget, OnGenerateWidget)
 	SLATE_END_ARGS()
 
@@ -67,13 +67,13 @@ public:
 
 private:
 	// Callback for changing the checked state of the check box.
-	void HandleCheckBoxCheckStateChanged( ESlateCheckBoxState::Type NewState )
+	void HandleCheckBoxCheckStateChanged( ECheckBoxState NewState )
 	{
 		OnCheckStateChanged.Execute(NewState);
 	}
 
 	// Callback for determining the checked state of the check box.
-	ESlateCheckBoxState::Type HandleCheckBoxIsChecked( ) const
+	ECheckBoxState HandleCheckBoxIsChecked( ) const
 	{
 		return IsChecked.Get();
 	}
@@ -84,7 +84,7 @@ public:
 private:
 	OptionType Option;
 	FOnCheckStateChanged OnCheckStateChanged;
-	TAttribute<ESlateCheckBoxState::Type> IsChecked;
+	TAttribute<ECheckBoxState> IsChecked;
 	FOnGenerateWidget OnGenerateWidget;
 };
 
@@ -126,10 +126,10 @@ private:
 	TSharedRef<SWidget> GenerateWidgetForOption(OptionType Option);
 
 	// Calls the provided user callback to notify them that an option has been checked or unchecked.
-	void HandleCheckBoxCheckStateChanged(ESlateCheckBoxState::Type NewState, OptionType Option) const;
+	void HandleCheckBoxCheckStateChanged(ECheckBoxState NewState, OptionType Option) const;
 
 	// Calls the provided user callback to be notified that an option is selected or unselected in some way.
-	ESlateCheckBoxState::Type HandleCheckBoxIsChecked(OptionType Option) const;
+	ECheckBoxState HandleCheckBoxIsChecked(OptionType Option) const;
 
 	// Selects all options.
 	void HandleAllHyperlinkNavigate();
@@ -263,15 +263,15 @@ TSharedRef<SWidget> SMultipleOptionTable<OptionType>::GenerateWidgetForOption(Op
 }
 
 template<typename OptionType>
-void SMultipleOptionTable<OptionType>::HandleCheckBoxCheckStateChanged(ESlateCheckBoxState::Type NewState, OptionType Option) const
+void SMultipleOptionTable<OptionType>::HandleCheckBoxCheckStateChanged(ECheckBoxState NewState, OptionType Option) const
 {
-	OnOptionSelectionChanged.Execute(NewState == ESlateCheckBoxState::Checked, Option);
+	OnOptionSelectionChanged.Execute(NewState == ECheckBoxState::Checked, Option);
 }
 
 template<typename OptionType>
-ESlateCheckBoxState::Type SMultipleOptionTable<OptionType>::HandleCheckBoxIsChecked(OptionType Option) const
+ECheckBoxState SMultipleOptionTable<OptionType>::HandleCheckBoxIsChecked(OptionType Option) const
 {
-	return IsOptionSelected.Execute(Option) ? ESlateCheckBoxState::Checked : ESlateCheckBoxState::Unchecked;
+	return IsOptionSelected.Execute(Option) ? ECheckBoxState::Checked : ECheckBoxState::Unchecked;
 }
 
 template<typename OptionType>
@@ -280,7 +280,7 @@ void SMultipleOptionTable<OptionType>::HandleAllHyperlinkNavigate()
 	OnPreBatchSelect.ExecuteIfBound();
 	for (const OptionType& Option : *Options)
 	{
-		HandleCheckBoxCheckStateChanged(ESlateCheckBoxState::Checked, Option);
+		HandleCheckBoxCheckStateChanged(ECheckBoxState::Checked, Option);
 	}
 	OnPostBatchSelect.ExecuteIfBound();
 }
@@ -291,7 +291,7 @@ void SMultipleOptionTable<OptionType>::HandleNoneHyperlinkNavigate()
 	OnPreBatchSelect.ExecuteIfBound();
 	for (const OptionType& Option : *Options)
 	{
-		HandleCheckBoxCheckStateChanged(ESlateCheckBoxState::Unchecked, Option);
+		HandleCheckBoxCheckStateChanged(ECheckBoxState::Unchecked, Option);
 	}
 	OnPostBatchSelect.ExecuteIfBound();
 }
diff --git a/Engine/Source/Runtime/AppFramework/Public/Widgets/Workflow/SWizard.h b/Engine/Source/Runtime/AppFramework/Public/Widgets/Workflow/SWizard.h
--- a/Engine/Source/Runtime/AppFramework/Public/Widgets/Workflow/SWizard.h
+++ b/Engine/Source/Runtime/AppFramework/Public/Widgets/Workflow/SWizard.h
@@ -264,10 +264,10 @@ private:
 	EVisibility HandleNextButtonVisibility() const;
 
 	// Callback for getting the checked state of a page button.
-	void HandlePageButtonCheckStateChanged( ESlateCheckBoxState::Type NewState, int32 PageIndex );
+	void HandlePageButtonCheckStateChanged( ECheckBoxState NewState, int32 PageIndex );
 
 	// Callback for clicking a page button.
-	ESlateCheckBoxState::Type HandlePageButtonIsChecked( int32 PageIndex ) const;
+	ECheckBoxState HandlePageButtonIsChecked( int32 PageIndex ) const;
 
 	// Callback for getting the enabled state of a page button.
 	bool HandlePageButtonIsEnabled( int32 PageIndex ) const;
diff --git a/Engine/Source/Runtime/Engine/Public/UnrealClient.h b/Engine/Source/Runtime/Engine/Public/UnrealClient.h
--- a/Engine/Source/Runtime/Engine/Public/UnrealClient.h
+++ b/Engine/Source/Runtime/Engine/Public/UnrealClient.h
@@ -15,7 +15,7 @@ class FCanvas;
 class FViewportClient;
 enum class EPopupMethod : uint8;
 
-enum class EFocusCause;
+enum class EFocusCause : uint8;
 
 /**
  * A render target.
diff --git a/Engine/Source/Runtime/Slate/Private/Framework/MultiBox/SButtonRowBlock.cpp b/Engine/Source/Runtime/Slate/Private/Framework/MultiBox/SButtonRowBlock.cpp
--- a/Engine/Source/Runtime/Slate/Private/Framework/MultiBox/SButtonRowBlock.cpp
+++ b/Engine/Source/Runtime/Slate/Private/Framework/MultiBox/SButtonRowBlock.cpp
@@ -238,7 +238,7 @@ FReply SButtonRowBlock::OnClicked()
 /**
  * Called by Slate when this check box button is toggled
  */
-void SButtonRowBlock::OnCheckStateChanged( const ESlateCheckBoxState::Type NewCheckedState )
+void SButtonRowBlock::OnCheckStateChanged( const ECheckBoxState NewCheckedState )
 {
 	OnClicked();
 }
@@ -246,9 +246,9 @@ void SButtonRowBlock::OnCheckStateChanged( const ESlateCheckBoxState::Type NewCh
 /**
  * Called by slate to determine if this button should appear checked
  *
- * @return ESlateCheckBoxState::Checked if it should be checked, ESlateCheckBoxState::Unchecked if not.
+ * @return ECheckBoxState::Checked if it should be checked, ECheckBoxState::Unchecked if not.
  */
-ESlateCheckBoxState::Type SButtonRowBlock::OnIsChecked() const
+ECheckBoxState SButtonRowBlock::OnIsChecked() const
 {
 	TSharedPtr< const FUICommandList > ActionList = MultiBlock->GetActionList();
 	TSharedPtr< const FUICommandInfo > Action = MultiBlock->GetAction();
@@ -265,7 +265,7 @@ ESlateCheckBoxState::Type SButtonRowBlock::OnIsChecked() const
 		bIsChecked = DirectActions.IsChecked();
 	}
 
-	return bIsChecked ? ESlateCheckBoxState::Checked : ESlateCheckBoxState::Unchecked;
+	return bIsChecked ? ECheckBoxState::Checked : ECheckBoxState::Unchecked;
 
 }
 
diff --git a/Engine/Source/Runtime/Slate/Private/Framework/MultiBox/SButtonRowBlock.h b/Engine/Source/Runtime/Slate/Private/Framework/MultiBox/SButtonRowBlock.h
--- a/Engine/Source/Runtime/Slate/Private/Framework/MultiBox/SButtonRowBlock.h
+++ b/Engine/Source/Runtime/Slate/Private/Framework/MultiBox/SButtonRowBlock.h
@@ -104,14 +104,14 @@ protected:
 	/**
 	 * Called by Slate when this tool bar check box button is toggled
 	 */
-	void OnCheckStateChanged( const ESlateCheckBoxState::Type NewCheckedState );
+	void OnCheckStateChanged( const ECheckBoxState NewCheckedState );
 
 	/**
 	 * Called by slate to determine if this button should appear checked
 	 *
-	 * @return ESlateCheckBoxState::Checked if it should be checked, ESlateCheckBoxState::Unchecked if not.
+	 * @return ECheckBoxState::Checked if it should be checked, ECheckBoxState::Unchecked if not.
 	 */
-	ESlateCheckBoxState::Type OnIsChecked() const;
+	ECheckBoxState OnIsChecked() const;
 
 	/**
 	 * Called by Slate to determine if this button is enabled
diff --git a/Engine/Source/Runtime/Slate/Private/Framework/MultiBox/SMenuEntryBlock.cpp b/Engine/Source/Runtime/Slate/Private/Framework/MultiBox/SMenuEntryBlock.cpp
--- a/Engine/Source/Runtime/Slate/Private/Framework/MultiBox/SMenuEntryBlock.cpp
+++ b/Engine/Source/Runtime/Slate/Private/Framework/MultiBox/SMenuEntryBlock.cpp
@@ -948,7 +948,7 @@ bool SMenuEntryBlock::IsEnabled() const
 /**
  * Called by Slate when this menu entry check box button is toggled
  */
-void SMenuEntryBlock::OnCheckStateChanged( const ESlateCheckBoxState::Type NewCheckedState )
+void SMenuEntryBlock::OnCheckStateChanged( const ECheckBoxState NewCheckedState )
 {
 	// The check box was clicked
 	const bool bCheckBoxClicked = true;
@@ -958,9 +958,9 @@ void SMenuEntryBlock::OnCheckStateChanged( const ESlateCheckBoxState::Type NewCh
 /**
  * Called by slate to determine if this menu entry should appear checked
  *
- * @return ESlateCheckBoxState::Checked if it should be checked, ESlateCheckBoxState::Unchecked if not.
+ * @return ECheckBoxState::Checked if it should be checked, ECheckBoxState::Unchecked if not.
  */
-ESlateCheckBoxState::Type SMenuEntryBlock::IsChecked() const
+ECheckBoxState SMenuEntryBlock::IsChecked() const
 {
 	TSharedPtr< const FUICommandList > ActionList = MultiBlock->GetActionList();
 	TSharedPtr< const FUICommandInfo > Action = MultiBlock->GetAction();
@@ -977,12 +977,12 @@ ESlateCheckBoxState::Type SMenuEntryBlock::IsChecked() const
 		bIsChecked = DirectActions.IsChecked();
 	}
 
-	return bIsChecked ? ESlateCheckBoxState::Checked : ESlateCheckBoxState::Unchecked;
+	return bIsChecked ? ECheckBoxState::Checked : ECheckBoxState::Unchecked;
 }
 
 const FSlateBrush* SMenuEntryBlock::OnGetCheckImage() const
 {
-	return IsChecked() == ESlateCheckBoxState::Checked ? CheckedImage : UncheckedImage;
+	return IsChecked() == ECheckBoxState::Checked ? CheckedImage : UncheckedImage;
 }
 
 /**
@@ -1250,8 +1250,8 @@ const FSlateBrush* SMenuEntryBlock::GetCheckBoxImageBrushFromStyle(const FCheckB
 {
 	switch (IsChecked())
 	{
-	case ESlateCheckBoxState::Checked: return &Style->CheckedImage;
-	case ESlateCheckBoxState::Unchecked: return &Style->UncheckedImage;
+	case ECheckBoxState::Checked: return &Style->CheckedImage;
+	case ECheckBoxState::Unchecked: return &Style->UncheckedImage;
 	default: return &Style->UndeterminedImage;
 	}
 }
diff --git a/Engine/Source/Runtime/Slate/Private/Framework/MultiBox/SMenuEntryBlock.h b/Engine/Source/Runtime/Slate/Private/Framework/MultiBox/SMenuEntryBlock.h
--- a/Engine/Source/Runtime/Slate/Private/Framework/MultiBox/SMenuEntryBlock.h
+++ b/Engine/Source/Runtime/Slate/Private/Framework/MultiBox/SMenuEntryBlock.h
@@ -210,14 +210,14 @@ protected:
 	/**
 	 * Called by Slate when this check box button is toggled in a menu entry
 	 */
-	void OnCheckStateChanged( const ESlateCheckBoxState::Type NewCheckedState );
+	void OnCheckStateChanged( const ECheckBoxState NewCheckedState );
 
 	/**
 	 * Called by slate to determine if this menu entry should appear checked
 	 *
 	 * @return true if it should be checked, false if not.
 	 */
-	ESlateCheckBoxState::Type IsChecked() const;
+	ECheckBoxState IsChecked() const;
 
 	/**
 	 * In the case that we have an icon to show.  This function is called to get the image that indicates the menu item should appear checked
diff --git a/Engine/Source/Runtime/Slate/Private/Framework/MultiBox/SToolBarButtonBlock.cpp b/Engine/Source/Runtime/Slate/Private/Framework/MultiBox/SToolBarButtonBlock.cpp
--- a/Engine/Source/Runtime/Slate/Private/Framework/MultiBox/SToolBarButtonBlock.cpp
+++ b/Engine/Source/Runtime/Slate/Private/Framework/MultiBox/SToolBarButtonBlock.cpp
@@ -317,7 +317,7 @@ FReply SToolBarButtonBlock::OnClicked()
 /**
  * Called by Slate when this tool bar check box button is toggled
  */
-void SToolBarButtonBlock::OnCheckStateChanged( const ESlateCheckBoxState::Type NewCheckedState )
+void SToolBarButtonBlock::OnCheckStateChanged( const ECheckBoxState NewCheckedState )
 {
 	OnClicked();
 }
@@ -325,9 +325,9 @@ void SToolBarButtonBlock::OnCheckStateChanged( const ESlateCheckBoxState::Type N
 /**
  * Called by slate to determine if this button should appear checked
  *
- * @return ESlateCheckBoxState::Checked if it should be checked, ESlateCheckBoxState::Unchecked if not.
+ * @return ECheckBoxState::Checked if it should be checked, ECheckBoxState::Unchecked if not.
  */
-ESlateCheckBoxState::Type SToolBarButtonBlock::OnIsChecked() const
+ECheckBoxState SToolBarButtonBlock::OnIsChecked() const
 {
 	TSharedPtr< const FUICommandList > ActionList = MultiBlock->GetActionList();
 	TSharedPtr< const FUICommandInfo > Action = MultiBlock->GetAction();
@@ -344,7 +344,7 @@ ESlateCheckBoxState::Type SToolBarButtonBlock::OnIsChecked() const
 		bIsChecked = DirectActions.IsChecked();
 	}
 
-	return bIsChecked ? ESlateCheckBoxState::Checked : ESlateCheckBoxState::Unchecked;
+	return bIsChecked ? ECheckBoxState::Checked : ECheckBoxState::Unchecked;
 }
 
 /**
diff --git a/Engine/Source/Runtime/Slate/Private/Widgets/Input/SCheckBox.cpp b/Engine/Source/Runtime/Slate/Private/Widgets/Input/SCheckBox.cpp
--- a/Engine/Source/Runtime/Slate/Private/Widgets/Input/SCheckBox.cpp
+++ b/Engine/Source/Runtime/Slate/Private/Widgets/Input/SCheckBox.cpp
@@ -117,12 +117,12 @@ FReply SCheckBox::OnKeyUp( const FGeometry& MyGeometry, const FKeyEvent& InKeyEv
 	if ( InKeyEvent.GetKey() == EKeys::SpaceBar )
 	{
 		ToggleCheckedState();
-		const TAttribute<ESlateCheckBoxState::Type>& State = IsCheckboxChecked.Get();
-		if(State == ESlateCheckBoxState::Checked)
+		const TAttribute<ECheckBoxState>& State = IsCheckboxChecked.Get();
+		if(State == ECheckBoxState::Checked)
 		{
 			PlayCheckedSound();
 		}
-		else if(State == ESlateCheckBoxState::Unchecked)
+		else if(State == ECheckBoxState::Unchecked)
 		{
 			PlayUncheckedSound();
 		}
@@ -150,12 +150,12 @@ FReply SCheckBox::OnMouseButtonDown( const FGeometry& MyGeometry, const FPointer
 		if( ClickMethod == EButtonClickMethod::MouseDown )
 		{
 			ToggleCheckedState();
-			const TAttribute<ESlateCheckBoxState::Type>& State = IsCheckboxChecked.Get();
-			if(State == ESlateCheckBoxState::Checked)
+			const TAttribute<ECheckBoxState>& State = IsCheckboxChecked.Get();
+			if(State == ECheckBoxState::Checked)
 			{
 				PlayCheckedSound();
 			}
-			else if(State == ESlateCheckBoxState::Unchecked)
+			else if(State == ECheckBoxState::Unchecked)
 			{
 				PlayUncheckedSound();
 			}
@@ -229,12 +229,12 @@ FReply SCheckBox::OnMouseButtonUp( const FGeometry& MyGeometry, const FPointerEv
 				if( ClickMethod == EButtonClickMethod::MouseUp || HasMouseCapture() )
 				{
 					ToggleCheckedState();
-					const TAttribute<ESlateCheckBoxState::Type>& State = IsCheckboxChecked.Get();
-					if(State == ESlateCheckBoxState::Checked)
+					const TAttribute<ECheckBoxState>& State = IsCheckboxChecked.Get();
+					if(State == ECheckBoxState::Checked)
 					{
 						PlayCheckedSound();
 					}
-					else if(State == ESlateCheckBoxState::Unchecked)
+					else if(State == ECheckBoxState::Unchecked)
 					{
 						PlayUncheckedSound();
 					}
@@ -275,21 +275,21 @@ void SCheckBox::OnMouseLeave( const FPointerEvent& MouseEvent )
  */
 const FSlateBrush* SCheckBox::OnGetCheckImage() const
 {
-	ESlateCheckBoxState::Type State = IsCheckboxChecked.Get();
+	ECheckBoxState State = IsCheckboxChecked.Get();
 
 	const FSlateBrush* ImageToUse;
 	switch( State )
 	{
-		case ESlateCheckBoxState::Unchecked:
+		case ECheckBoxState::Unchecked:
 			ImageToUse = IsPressed() ? GetUncheckedPressedImage() : ( IsHovered() ? GetUncheckedHoveredImage() : GetUncheckedImage() );
 			break;
 	
-		case ESlateCheckBoxState::Checked:
+		case ECheckBoxState::Checked:
 			ImageToUse = IsPressed() ? GetCheckedPressedImage() : ( IsHovered() ? GetCheckedHoveredImage() : GetCheckedImage() );
 			break;
 	
 		default:
-		case ESlateCheckBoxState::Undetermined:
+		case ECheckBoxState::Undetermined:
 			ImageToUse = IsPressed() ? GetUndeterminedPressedImage() : ( IsHovered() ? GetUndeterminedHoveredImage() : GetUndeterminedImage() );
 			break;
 	}
@@ -298,7 +298,7 @@ const FSlateBrush* SCheckBox::OnGetCheckImage() const
 }
 
 
-ESlateCheckBoxState::Type SCheckBox::GetCheckedState() const
+ECheckBoxState SCheckBox::GetCheckedState() const
 {
 	return IsCheckboxChecked.Get();
 }
@@ -308,34 +308,34 @@ ESlateCheckBoxState::Type SCheckBox::GetCheckedState() const
  */
 void SCheckBox::ToggleCheckedState()
 {
-	const TAttribute<ESlateCheckBoxState::Type>& State = IsCheckboxChecked.Get();
+	const TAttribute<ECheckBoxState>& State = IsCheckboxChecked.Get();
 
 	// If the current check box state is checked OR undetermined we set the check box to checked.
-	if( State == ESlateCheckBoxState::Checked || State == ESlateCheckBoxState::Undetermined )
+	if( State == ECheckBoxState::Checked || State == ECheckBoxState::Undetermined )
 	{
 		if ( !IsCheckboxChecked.IsBound() )
 		{
 			// When we are not bound, just toggle the current state.
-			IsCheckboxChecked.Set( ESlateCheckBoxState::Unchecked );
+			IsCheckboxChecked.Set( ECheckBoxState::Unchecked );
 		}
 
 		// The state of the check box changed.  Execute the delegate to notify users
-		OnCheckStateChanged.ExecuteIfBound( ESlateCheckBoxState::Unchecked );
+		OnCheckStateChanged.ExecuteIfBound( ECheckBoxState::Unchecked );
 	}
-	else if( State == ESlateCheckBoxState::Unchecked )
+	else if( State == ECheckBoxState::Unchecked )
 	{
 		if ( !IsCheckboxChecked.IsBound() )
 		{
 			// When we are not bound, just toggle the current state.
-			IsCheckboxChecked.Set( ESlateCheckBoxState::Checked );
+			IsCheckboxChecked.Set( ECheckBoxState::Checked );
 		}
 
 		// The state of the check box changed.  Execute the delegate to notify users
-		OnCheckStateChanged.ExecuteIfBound( ESlateCheckBoxState::Checked );
+		OnCheckStateChanged.ExecuteIfBound( ECheckBoxState::Checked );
 	}
 }
 
-void SCheckBox::SetIsChecked(TAttribute<ESlateCheckBoxState::Type> InIsChecked)
+void SCheckBox::SetIsChecked(TAttribute<ECheckBoxState> InIsChecked)
 {
 	IsCheckboxChecked = InIsChecked;
 }
diff --git a/Engine/Source/Runtime/Slate/Private/Widgets/Notifications/SNotificationList.cpp b/Engine/Source/Runtime/Slate/Private/Widgets/Notifications/SNotificationList.cpp
--- a/Engine/Source/Runtime/Slate/Private/Widgets/Notifications/SNotificationList.cpp
+++ b/Engine/Source/Runtime/Slate/Private/Widgets/Notifications/SNotificationList.cpp
@@ -298,7 +298,7 @@ public:
 		/** When set this forces the width of the box, used to stop resizeing on text change */
 		SLATE_ARGUMENT( FOptionalSize, WidthOverride )
 		/** When set this will display a check box on the notification; handles getting the current check box state */
-		SLATE_ATTRIBUTE( ESlateCheckBoxState::Type, CheckBoxState )
+		SLATE_ATTRIBUTE( ECheckBoxState, CheckBoxState )
 		/** When set this will display a check box on the notification; handles setting the new check box state */
 		SLATE_EVENT( FOnCheckStateChanged, CheckBoxStateChanged );
 		/** Text to display for the check box message */
diff --git a/Engine/Source/Runtime/Slate/Public/Framework/MultiBox/SToolBarButtonBlock.h b/Engine/Source/Runtime/Slate/Public/Framework/MultiBox/SToolBarButtonBlock.h
--- a/Engine/Source/Runtime/Slate/Public/Framework/MultiBox/SToolBarButtonBlock.h
+++ b/Engine/Source/Runtime/Slate/Public/Framework/MultiBox/SToolBarButtonBlock.h
@@ -124,14 +124,14 @@ protected:
 	/**
 	 * Called by Slate when this tool bar check box button is toggled
 	 */
-	void OnCheckStateChanged( const ESlateCheckBoxState::Type NewCheckedState );
+	void OnCheckStateChanged( const ECheckBoxState NewCheckedState );
 
 	/**
 	 * Called by slate to determine if this button should appear checked
 	 *
-	 * @return ESlateCheckBoxState::Checked if it should be checked, ESlateCheckBoxState::Unchecked if not.
+	 * @return ECheckBoxState::Checked if it should be checked, ECheckBoxState::Unchecked if not.
 	 */
-	ESlateCheckBoxState::Type OnIsChecked() const;
+	ECheckBoxState OnIsChecked() const;
 
 	/**
 	 * Called by Slate to determine if this button is enabled
diff --git a/Engine/Source/Runtime/Slate/Public/Widgets/Input/SCheckBox.h b/Engine/Source/Runtime/Slate/Public/Widgets/Input/SCheckBox.h
--- a/Engine/Source/Runtime/Slate/Public/Widgets/Input/SCheckBox.h
+++ b/Engine/Source/Runtime/Slate/Public/Widgets/Input/SCheckBox.h
@@ -8,22 +8,31 @@
 
 /** Current state of the check box */
 UENUM(BlueprintType)
+enum class ECheckBoxState : uint8
+{
+	/** Unchecked */
+	Unchecked,
+	/** Checked */
+	Checked,
+	/** Neither checked nor unchecked */
+	Undetermined
+};
+
+/** DEPRECATED 4.6 - Do not use */
+//@Todo slate: Remove this as soon as the 4.6 deprecated API is Removed.
 namespace ESlateCheckBoxState
 {
-	enum Type
-	{
-		/** Unchecked */
-		Unchecked,
-		/** Checked */
-		Checked,
-		/** Neither checked nor unchecked */
-		Undetermined
-	};
+	DEPRECATED(4.6, "ESlateCheckBoxState::Type is deprecated and was renamed to ECheckBoxState. Please use that type instead.")
+	typedef ::ECheckBoxState Type;
+
+	const ECheckBoxState Unchecked = ECheckBoxState::Unchecked;
+	const ECheckBoxState Checked = ECheckBoxState::Checked;
+	const ECheckBoxState Undetermined = ECheckBoxState::Undetermined;
 }
 
 
 /** Delegate that is executed when the check box state changes */
-DECLARE_DELEGATE_OneParam( FOnCheckStateChanged, ESlateCheckBoxState::Type );
+DECLARE_DELEGATE_OneParam( FOnCheckStateChanged, ECheckBoxState );
 
 
 /**
@@ -38,7 +47,7 @@ public:
 		, _Style( &FCoreStyle::Get().GetWidgetStyle< FCheckBoxStyle >("Checkbox") )
 		, _Type()
 		, _OnCheckStateChanged()
-		, _IsChecked( ESlateCheckBoxState::Unchecked )
+		, _IsChecked( ECheckBoxState::Unchecked )
 		, _HAlign( HAlign_Fill )
 		, _Padding()
 		, _ClickMethod( EButtonClickMethod::DownAndUp )
@@ -69,21 +78,21 @@ public:
 		SLATE_EVENT( FOnCheckStateChanged, OnCheckStateChanged )
 
 		/** Whether the check box is currently in a checked state */
-		SLATE_ATTRIBUTE( ESlateCheckBoxState::Type, IsChecked )
+		SLATE_ATTRIBUTE( ECheckBoxState, IsChecked )
 
 		// TODO Remove this function when IsChecked(bool InIsChecked) is removed.  It has to be here to prevent ambiguous conversions.
 		/** Whether the check box is currently in a checked state */
-		FArguments& IsChecked(ESlateCheckBoxState::Type InIsChecked)
+		FArguments& IsChecked(ECheckBoxState InIsChecked)
 		{
 			_IsChecked = InIsChecked;
 			return Me();
 		}
 
 		/** Whether the check box is currently in a checked state */
-		DEPRECATED(4.3, "Please use IsChecked(TAttribute<ESlateCheckBoxState::Type>)")
+		DEPRECATED(4.3, "Please use IsChecked(TAttribute<ECheckBoxState>)")
 		FArguments& IsChecked(bool InIsChecked)
 		{
-			_IsChecked = InIsChecked ? ESlateCheckBoxState::Checked : ESlateCheckBoxState::Unchecked;
+			_IsChecked = InIsChecked ? ECheckBoxState::Checked : ECheckBoxState::Unchecked;
 			return Me();
 		}
 
@@ -169,11 +178,11 @@ public:
 	 */
 	bool IsChecked() const
 	{
-		return (IsCheckboxChecked.Get() == ESlateCheckBoxState::Checked);
+		return ( IsCheckboxChecked.Get() == ECheckBoxState::Checked );
 	}
 
 	/** @return The current checked state of the checkbox. */
-	ESlateCheckBoxState::Type GetCheckedState() const;
+	ECheckBoxState GetCheckedState() const;
 
 	/**
 	 * Returns true if this button is currently pressed
@@ -191,7 +200,7 @@ public:
 	void ToggleCheckedState();
 
 	/** See the IsChecked attribute */
-	void SetIsChecked(TAttribute<ESlateCheckBoxState::Type> InIsChecked);
+	void SetIsChecked(TAttribute<ECheckBoxState> InIsChecked);
 	
 	/** See the Content slot */
 	void SetContent(const TSharedRef< SWidget >& InContent);
@@ -248,7 +257,7 @@ protected:
 	bool bIsPressed;
 
 	/** Are we checked */
-	TAttribute<ESlateCheckBoxState::Type> IsCheckboxChecked;
+	TAttribute<ECheckBoxState> IsCheckboxChecked;
 
 	/** Delegate called when the check box changes state */
 	FOnCheckStateChanged OnCheckStateChanged;
diff --git a/Engine/Source/Runtime/Slate/Public/Widgets/Notifications/SNotificationList.h b/Engine/Source/Runtime/Slate/Public/Widgets/Notifications/SNotificationList.h
--- a/Engine/Source/Runtime/Slate/Public/Widgets/Notifications/SNotificationList.h
+++ b/Engine/Source/Runtime/Slate/Public/Widgets/Notifications/SNotificationList.h
@@ -109,7 +109,7 @@ struct FNotificationInfo
 		bUseLargeFont(true),
 		WidthOverride(),
 		bFireAndForget(true),
-		CheckBoxState(ESlateCheckBoxState::Unchecked),
+		CheckBoxState(ECheckBoxState::Unchecked),
 		CheckBoxStateChanged(),
 		CheckBoxText(),
 		Hyperlink(),
@@ -135,7 +135,7 @@ struct FNotificationInfo
 		bUseLargeFont(false),
 		WidthOverride(),
 		bFireAndForget(true),
-		CheckBoxState(ESlateCheckBoxState::Unchecked),
+		CheckBoxState(ECheckBoxState::Unchecked),
 		CheckBoxStateChanged(),
 		CheckBoxText(),
 		Hyperlink(),
@@ -180,7 +180,7 @@ struct FNotificationInfo
 	bool bFireAndForget;
 
 	/** When set this will display a check box on the notification; handles getting the current check box state */
-	TAttribute< ESlateCheckBoxState::Type > CheckBoxState;
+	TAttribute< ECheckBoxState > CheckBoxState;
 
 	/** When set this will display a check box on the notification; handles setting the new check box state */
 	FOnCheckStateChanged CheckBoxStateChanged;
diff --git a/Engine/Source/Runtime/SlateCore/Public/Input/Events.h b/Engine/Source/Runtime/SlateCore/Public/Input/Events.h
--- a/Engine/Source/Runtime/SlateCore/Public/Input/Events.h
+++ b/Engine/Source/Runtime/SlateCore/Public/Input/Events.h
@@ -13,7 +13,7 @@ class SWindow;
 * Context for focus change
 */
 UENUM()
-enum class EFocusCause
+enum class EFocusCause : uint8
 {
 	/** Focus was changed because of a mouse action */
 	Mouse,
diff --git a/Engine/Source/Runtime/UMG/Private/Binding/CheckedStateBinding.cpp b/Engine/Source/Runtime/UMG/Private/Binding/CheckedStateBinding.cpp
--- a/Engine/Source/Runtime/UMG/Private/Binding/CheckedStateBinding.cpp
+++ b/Engine/Source/Runtime/UMG/Private/Binding/CheckedStateBinding.cpp
@@ -29,7 +29,7 @@ bool UCheckedStateBinding::IsSupportedDestination(UProperty* Property) const
 	return false;
 }
 
-ESlateCheckBoxState::Type UCheckedStateBinding::GetValue() const
+ECheckBoxState UCheckedStateBinding::GetValue() const
 {
 	if ( UObject* Source = SourceObject.Get() )
 	{
@@ -39,7 +39,7 @@ ESlateCheckBoxState::Type UCheckedStateBinding::GetValue() const
 			if ( SourcePath.GetValue<uint8>(Source, Value) )
 			{
 				bConversion = EConversion::None;
-				return static_cast<ESlateCheckBoxState::Type>(Value);
+				return static_cast<ECheckBoxState>(Value);
 			}
 		}
 
@@ -49,12 +49,12 @@ ESlateCheckBoxState::Type UCheckedStateBinding::GetValue() const
 			if ( SourcePath.GetValue<bool>(Source, Value) )
 			{
 				bConversion = EConversion::Bool;
-				return Value ? ESlateCheckBoxState::Checked : ESlateCheckBoxState::Unchecked;
+				return Value ? ECheckBoxState::Checked : ECheckBoxState::Unchecked;
 			}
 		}
 	}
 
-	return ESlateCheckBoxState::Unchecked;
+	return ECheckBoxState::Unchecked;
 }
 
 #undef LOCTEXT_NAMESPACE
diff --git a/Engine/Source/Runtime/UMG/Private/Components/CheckBox.cpp b/Engine/Source/Runtime/UMG/Private/Components/CheckBox.cpp
--- a/Engine/Source/Runtime/UMG/Private/Components/CheckBox.cpp
+++ b/Engine/Source/Runtime/UMG/Private/Components/CheckBox.cpp
@@ -13,7 +13,7 @@ UCheckBox::UCheckBox(const FObjectInitializer& ObjectInitializer)
 	SCheckBox::FArguments SlateDefaults;
 	WidgetStyle = *SlateDefaults._Style;
 
-	CheckedState = ESlateCheckBoxState::Unchecked;
+	CheckedState = ECheckBoxState::Unchecked;
 
 	HorizontalAlignment = SlateDefaults._HAlign;
 	Padding = SlateDefaults._Padding.Get();
@@ -50,7 +50,7 @@ void UCheckBox::SynchronizeProperties()
 {
 	Super::SynchronizeProperties();
 
-	MyCheckbox->SetIsChecked( OPTIONAL_BINDING(ESlateCheckBoxState::Type, CheckedState) );
+	MyCheckbox->SetIsChecked( OPTIONAL_BINDING(ECheckBoxState, CheckedState) );
 }
 
 void UCheckBox::OnSlotAdded(UPanelSlot* Slot)
@@ -88,10 +88,10 @@ bool UCheckBox::IsChecked() const
 		return MyCheckbox->IsChecked();
 	}
 
-	return ( CheckedState == ESlateCheckBoxState::Checked );
+	return ( CheckedState == ECheckBoxState::Checked );
 }
 
-ESlateCheckBoxState::Type UCheckBox::GetCheckedState() const
+ECheckBoxState UCheckBox::GetCheckedState() const
 {
 	if ( MyCheckbox.IsValid() )
 	{
@@ -103,26 +103,26 @@ ESlateCheckBoxState::Type UCheckBox::GetCheckedState() const
 
 void UCheckBox::SetIsChecked(bool InIsChecked)
 {
-	CheckedState = InIsChecked ? ESlateCheckBoxState::Checked : ESlateCheckBoxState::Unchecked;
+	CheckedState = InIsChecked ? ECheckBoxState::Checked : ECheckBoxState::Unchecked;
 	if ( MyCheckbox.IsValid() )
 	{
-		MyCheckbox->SetIsChecked(OPTIONAL_BINDING(ESlateCheckBoxState::Type, CheckedState));
+		MyCheckbox->SetIsChecked(OPTIONAL_BINDING(ECheckBoxState, CheckedState));
 	}
 }
 
-void UCheckBox::SetCheckedState(ESlateCheckBoxState::Type InCheckedState)
+void UCheckBox::SetCheckedState(ECheckBoxState InCheckedState)
 {
 	CheckedState = InCheckedState;
 	if ( MyCheckbox.IsValid() )
 	{
-		MyCheckbox->SetIsChecked(OPTIONAL_BINDING(ESlateCheckBoxState::Type, CheckedState));
+		MyCheckbox->SetIsChecked(OPTIONAL_BINDING(ECheckBoxState, CheckedState));
 	}
 }
 
-void UCheckBox::SlateOnCheckStateChangedCallback(ESlateCheckBoxState::Type NewState)
+void UCheckBox::SlateOnCheckStateChangedCallback(ECheckBoxState NewState)
 {
 	//@TODO: Choosing to treat Undetermined as Checked
-	const bool bWantsToBeChecked = NewState != ESlateCheckBoxState::Unchecked;
+	const bool bWantsToBeChecked = NewState != ECheckBoxState::Unchecked;
 
 	if (OnCheckStateChanged.IsBound())
 	{
diff --git a/Engine/Source/Runtime/UMG/Public/Binding/CheckedStateBinding.h b/Engine/Source/Runtime/UMG/Public/Binding/CheckedStateBinding.h
--- a/Engine/Source/Runtime/UMG/Public/Binding/CheckedStateBinding.h
+++ b/Engine/Source/Runtime/UMG/Public/Binding/CheckedStateBinding.h
@@ -17,7 +17,7 @@ public:
 	virtual bool IsSupportedDestination(UProperty* Property) const override;
 
 	UFUNCTION()
-	ESlateCheckBoxState::Type GetValue() const;
+	ECheckBoxState GetValue() const;
 
 private:
 	enum class EConversion : uint8
diff --git a/Engine/Source/Runtime/UMG/Public/Components/CheckBox.h b/Engine/Source/Runtime/UMG/Public/Components/CheckBox.h
--- a/Engine/Source/Runtime/UMG/Public/Components/CheckBox.h
+++ b/Engine/Source/Runtime/UMG/Public/Components/CheckBox.h
@@ -24,7 +24,7 @@ class UMG_API UCheckBox : public UContentWidget
 public:
 	/** Whether the check box is currently in a checked state */
 	UPROPERTY(EditDefaultsOnly, Category=Appearance)
-	TEnumAsByte<ESlateCheckBoxState::Type> CheckedState;
+	ECheckBoxState CheckedState;
 
 	/** A bindable delegate for the IsChecked. */
 	UPROPERTY()
@@ -105,7 +105,7 @@ public:
 
 	/** @return the full current checked state. */
 	UFUNCTION(BlueprintCallable, Category="Widget")
-	ESlateCheckBoxState::Type GetCheckedState() const;
+	ECheckBoxState GetCheckedState() const;
 
 	/** Sets the checked state. */
 	UFUNCTION(BlueprintCallable, Category="Widget")
@@ -113,7 +113,7 @@ public:
 
 	/** Sets the checked state. */
 	UFUNCTION(BlueprintCallable, Category="Widget")
-	void SetCheckedState(ESlateCheckBoxState::Type InCheckedState);
+	void SetCheckedState(ECheckBoxState InCheckedState);
 
 public:
 	
@@ -146,7 +146,7 @@ protected:
 	virtual TSharedRef<SWidget> RebuildWidget() override;
 	// End of UWidget interface
 
-	void SlateOnCheckStateChangedCallback(ESlateCheckBoxState::Type NewState);
+	void SlateOnCheckStateChangedCallback(ECheckBoxState NewState);
 	
 protected:
 	TSharedPtr<SCheckBox> MyCheckbox;
diff --git a/Engine/Source/Runtime/UMG/Public/Components/Widget.h b/Engine/Source/Runtime/UMG/Public/Components/Widget.h
--- a/Engine/Source/Runtime/UMG/Public/Components/Widget.h
+++ b/Engine/Source/Runtime/UMG/Public/Components/Widget.h
@@ -41,10 +41,10 @@ public:
 	//DECLARE_DYNAMIC_DELEGATE_RetVal(FMargin, FGetMargin);
 	DECLARE_DYNAMIC_DELEGATE_RetVal(FSlateColor, FGetSlateColor);
 	DECLARE_DYNAMIC_DELEGATE_RetVal(FLinearColor, FGetLinearColor);
+	DECLARE_DYNAMIC_DELEGATE_RetVal(FSlateBrush, FGetSlateBrush);
 	DECLARE_DYNAMIC_DELEGATE_RetVal(ESlateVisibility, FGetSlateVisibility);
 	DECLARE_DYNAMIC_DELEGATE_RetVal(EMouseCursor::Type, FGetMouseCursor);
-	DECLARE_DYNAMIC_DELEGATE_RetVal(FSlateBrush, FGetSlateBrush);
-	DECLARE_DYNAMIC_DELEGATE_RetVal(ESlateCheckBoxState::Type, FGetCheckBoxState);
+	DECLARE_DYNAMIC_DELEGATE_RetVal(ECheckBoxState, FGetCheckBoxState);
 	DECLARE_DYNAMIC_DELEGATE_RetVal(UWidget*, FGetContent);
 
 	DECLARE_DYNAMIC_DELEGATE_RetVal_OneParam(UWidget*, FGenerateWidgetForString, FString, Item);
