Slate Test Suite : WidgetGallery is now the face of of the test suite; it is a nomad tab and comes up first.

diff --git a/Engine/Source/Runtime/AppFramework/Private/Framework/Testing/SWidgetGallery.cpp b/Engine/Source/Runtime/AppFramework/Private/Framework/Testing/SWidgetGallery.cpp
--- a/Engine/Source/Runtime/AppFramework/Private/Framework/Testing/SWidgetGallery.cpp
+++ b/Engine/Source/Runtime/AppFramework/Private/Framework/Testing/SWidgetGallery.cpp
@@ -71,754 +71,758 @@ public:
 		
 		ChildSlot
 		[
-			SNew(SScrollBox)
-
-			+ SScrollBox::Slot()
-				.Padding(5.0f)
-				[
-					SNew(SGridPanel)
-						.FillColumn(0, 0.5f)
-						.FillColumn(1, 0.5f)
-
-					// SBorder
-					+ SGridPanel::Slot(0, 0)
-						[
-							SNew(STextBlock)
-								.Text(LOCTEXT("SBorderLabel", "SBorder"))
-						]
-
-					+ SGridPanel::Slot(1, 0)
-						.Padding(0.0f, 5.0f)
-						[
-							SNew(SBorder)
-								[
-									SNew(SSpacer)
-										.Size(FVector2D(100.0f, 50.0f))
-								]	
-						]
-
-					// SBreadcrumbTrailLabel
-					+ SGridPanel::Slot(0, 1)
-						[
-							SNew(STextBlock)
-								.Text(LOCTEXT("SBreadcrumbTrailLabel", "SBreadcrumbTrail"))
-						]
-
-					+ SGridPanel::Slot(1, 1)
-						.Padding(0.0f, 5.0f)
-						[
-							SNew(SHorizontalBox)
+			SNew(SBorder)
+			.BorderImage( FCoreStyle::Get().GetBrush("ToolPanel.GroupBorder") )
+			[
+				SNew(SScrollBox)
 
-							+ SHorizontalBox::Slot()
-								.FillWidth(1.0f)
-								[
-									SAssignNew(BreadcrumbTrail, SBreadcrumbTrail<int32>)
+				+ SScrollBox::Slot()
+					.Padding(5.0f)
+					[
+						SNew(SGridPanel)
+							.FillColumn(0, 0.5f)
+							.FillColumn(1, 0.5f)
 
-								]
+						// SBorder
+						+ SGridPanel::Slot(0, 0)
+							[
+								SNew(STextBlock)
+									.Text(LOCTEXT("SBorderLabel", "SBorder"))
+							]
 
-							+ SHorizontalBox::Slot()
-								.AutoWidth()
-								[
-									SNew(SButton)
-										.Text(LOCTEXT("AddBreadCrumbLabel", "Add"))
-										.HAlign(HAlign_Center)
-										.VAlign(VAlign_Center)
-										.OnClicked(this, &SWidgetGallery::HandleBreadcrumbTrailAddButtonClicked)
-								]
-						]
+						+ SGridPanel::Slot(1, 0)
+							.Padding(0.0f, 5.0f)
+							[
+								SNew(SBorder)
+									[
+										SNew(SSpacer)
+											.Size(FVector2D(100.0f, 50.0f))
+									]	
+							]
 
-					// SButton
-					+ SGridPanel::Slot(0, 2)
-						[
-							SNew(STextBlock)
-								.Text(LOCTEXT("SButtonLabel", "SButton"))
-						]
+						// SBreadcrumbTrailLabel
+						+ SGridPanel::Slot(0, 1)
+							[
+								SNew(STextBlock)
+									.Text(LOCTEXT("SBreadcrumbTrailLabel", "SBreadcrumbTrail"))
+							]
 
-					+ SGridPanel::Slot(1, 2)
-						.HAlign(HAlign_Left)
-						.Padding(0.0f, 5.0f)
-						[
-							SNew(SButton)
-								.Text(LOCTEXT("ButtonExampleLabel", "Button"))
-						]
+						+ SGridPanel::Slot(1, 1)
+							.Padding(0.0f, 5.0f)
+							[
+								SNew(SHorizontalBox)
 
-					// SButton (no content)
-					+ SGridPanel::Slot(0, 3)
-						[
-							SNew(STextBlock)
-								.Text(LOCTEXT("SButtonExampleLabel", "SButton (no content)"))
-						]
+								+ SHorizontalBox::Slot()
+									.FillWidth(1.0f)
+									[
+										SAssignNew(BreadcrumbTrail, SBreadcrumbTrail<int32>)
 
-					+ SGridPanel::Slot(1, 3)
-						.HAlign(HAlign_Left)
-						.Padding(0.0f, 5.0f)
-						[
-							SNew(SButton)
-						]
+									]
 
-					// SCheckBox
-					+ SGridPanel::Slot(0, 4)
-						[
-							SNew(STextBlock)
-								.Text(LOCTEXT("SCheckBoxLabel", "SCheckBox"))
-						]
+								+ SHorizontalBox::Slot()
+									.AutoWidth()
+									[
+										SNew(SButton)
+											.Text(LOCTEXT("AddBreadCrumbLabel", "Add"))
+											.HAlign(HAlign_Center)
+											.VAlign(VAlign_Center)
+											.OnClicked(this, &SWidgetGallery::HandleBreadcrumbTrailAddButtonClicked)
+									]
+							]
 
-					+ SGridPanel::Slot(1, 4)
-						.HAlign(HAlign_Left)
-						.Padding(0.0f, 5.0f)
-						[
-							SNew(SVerticalBox)
+						// SButton
+						+ SGridPanel::Slot(0, 2)
+							[
+								SNew(STextBlock)
+									.Text(LOCTEXT("SButtonLabel", "SButton"))
+							]
 
-							+ SVerticalBox::Slot()
-								.AutoHeight()
-								[
-									CreateCheckBox(LOCTEXT("SCheckBoxItemLabel01", "Option 1"), &CheckBox1Choice)
-								]
+						+ SGridPanel::Slot(1, 2)
+							.HAlign(HAlign_Left)
+							.Padding(0.0f, 5.0f)
+							[
+								SNew(SButton)
+									.Text(LOCTEXT("ButtonExampleLabel", "Button"))
+							]
 
-							+ SVerticalBox::Slot()
-								.AutoHeight()
-								[
-									CreateCheckBox(LOCTEXT("SCheckBoxItemLabel02", "Option 2"), &CheckBox2Choice)
-								]
+						// SButton (no content)
+						+ SGridPanel::Slot(0, 3)
+							[
+								SNew(STextBlock)
+									.Text(LOCTEXT("SButtonExampleLabel", "SButton (no content)"))
+							]
 
-							+ SVerticalBox::Slot()
-								.AutoHeight()
-								[
-									CreateCheckBox(LOCTEXT("SCheckBoxItemLabel03", "Option 3"), &CheckBox3Choice)
-								]
-						]
+						+ SGridPanel::Slot(1, 3)
+							.HAlign(HAlign_Left)
+							.Padding(0.0f, 5.0f)
+							[
+								SNew(SButton)
+							]
 
-					// SCheckBox (as radio button)
-					+ SGridPanel::Slot(0, 5)
-						[
-							SNew(STextBlock)
-								.Text(LOCTEXT("SRadioButtonLabel", "SCheckBox (as radio buttons)"))
-						]
+						// SCheckBox
+						+ SGridPanel::Slot(0, 4)
+							[
+								SNew(STextBlock)
+									.Text(LOCTEXT("SCheckBoxLabel", "SCheckBox"))
+							]
 
-					+ SGridPanel::Slot(1, 5)
-						.HAlign(HAlign_Left)
-						.Padding(0.0f, 5.0f)
-						[
-							SNew(SVerticalBox)
+						+ SGridPanel::Slot(1, 4)
+							.HAlign(HAlign_Left)
+							.Padding(0.0f, 5.0f)
+							[
+								SNew(SVerticalBox)
 
-							+ SVerticalBox::Slot()
-								.AutoHeight()
-								[
-									CreateRadioButton(LOCTEXT("SRadioButtonItemLabel01", "Option 1"), Radio0)
-								]
+								+ SVerticalBox::Slot()
+									.AutoHeight()
+									[
+										CreateCheckBox(LOCTEXT("SCheckBoxItemLabel01", "Option 1"), &CheckBox1Choice)
+									]
 
-							+ SVerticalBox::Slot()
-								.AutoHeight()
-								[
-									CreateRadioButton(LOCTEXT("SRadioButtonItemLabel02", "Option 2"), Radio1)
-								]
+								+ SVerticalBox::Slot()
+									.AutoHeight()
+									[
+										CreateCheckBox(LOCTEXT("SCheckBoxItemLabel02", "Option 2"), &CheckBox2Choice)
+									]
 
-							+ SVerticalBox::Slot()
-								.AutoHeight()
-								[
-									CreateRadioButton(LOCTEXT("SRadioButtonItemLabel03", "Option 3"), Radio2)
-								]
-						]
+								+ SVerticalBox::Slot()
+									.AutoHeight()
+									[
+										CreateCheckBox(LOCTEXT("SCheckBoxItemLabel03", "Option 3"), &CheckBox3Choice)
+									]
+							]
 
-					// SCircularThrobber
-					+ SGridPanel::Slot(0, 6)
-						[
-							SNew(STextBlock)
-								.Text(LOCTEXT("SCircularThrobberLabel", "SCircularThrobber"))
-						]
+						// SCheckBox (as radio button)
+						+ SGridPanel::Slot(0, 5)
+							[
+								SNew(STextBlock)
+									.Text(LOCTEXT("SRadioButtonLabel", "SCheckBox (as radio buttons)"))
+							]
 
-					+ SGridPanel::Slot(1, 6)
-						.HAlign(HAlign_Left)
-						.Padding(0.0f, 5.0f)
-						[
-							SNew(SCircularThrobber)
-						]
+						+ SGridPanel::Slot(1, 5)
+							.HAlign(HAlign_Left)
+							.Padding(0.0f, 5.0f)
+							[
+								SNew(SVerticalBox)
 
-					// SColorBlock
-					+ SGridPanel::Slot(0, 7)
-						[
-							SNew(STextBlock)
-								.Text(LOCTEXT("SColorBlockLabel", "SColorBlock"))
-						]
+								+ SVerticalBox::Slot()
+									.AutoHeight()
+									[
+										CreateRadioButton(LOCTEXT("SRadioButtonItemLabel01", "Option 1"), Radio0)
+									]
 
-					+ SGridPanel::Slot(1, 7)
-						.Padding(0.0f, 5.0f)
-						[
-							SNew(SColorBlock)
-								.Color(FLinearColor(1.0f, 0.0f, 0.0f))
-						]
+								+ SVerticalBox::Slot()
+									.AutoHeight()
+									[
+										CreateRadioButton(LOCTEXT("SRadioButtonItemLabel02", "Option 2"), Radio1)
+									]
 
-					// SComboBox
-					+ SGridPanel::Slot(0, 8)
-						[
-							SNew(STextBlock)
-								.Text(LOCTEXT("SComboBoxLabel", "SComboBox"))
-						]
+								+ SVerticalBox::Slot()
+									.AutoHeight()
+									[
+										CreateRadioButton(LOCTEXT("SRadioButtonItemLabel03", "Option 3"), Radio2)
+									]
+							]
 
-					+ SGridPanel::Slot(1, 8)
-						.HAlign(HAlign_Left)
-						.Padding(0.0f, 5.0f)
-						[
-							SNew(SVerticalBox)
+						// SCircularThrobber
+						+ SGridPanel::Slot(0, 6)
+							[
+								SNew(STextBlock)
+									.Text(LOCTEXT("SCircularThrobberLabel", "SCircularThrobber"))
+							]
 
-							+ SVerticalBox::Slot()
+						+ SGridPanel::Slot(1, 6)
+							.HAlign(HAlign_Left)
+							.Padding(0.0f, 5.0f)
 							[
-								SNew(SCheckBox)
-								.IsChecked(this, &SWidgetGallery::IsCreateNewWindowChecked)
-								.OnCheckStateChanged( this, &SWidgetGallery::CreateNewWindowToggled )
-								[
-									SNew(STextBlock).Text(LOCTEXT("CreateNewWindowCheckbox", "Create new window for popups."))
-								]
+								SNew(SCircularThrobber)
 							]
 
-							+ SVerticalBox::Slot()
-							.AutoHeight()
+						// SColorBlock
+						+ SGridPanel::Slot(0, 7)
 							[
-								// selector combo box
-								SAssignNew(SelectorComboBox, SComboBox<TSharedPtr<FString> >)
-									.OptionsSource(&SelectorComboBoxOptions)
-									.OnSelectionChanged(this, &SWidgetGallery::HandleSelectorComboBoxSelectionChanged)
-									.OnGenerateWidget(this, &SWidgetGallery::HandleComboBoxGenerateWidget)
-									[
-										SNew(STextBlock)
-											.Text(this, &SWidgetGallery::HandleSelectorComboBoxText)
-									]
+								SNew(STextBlock)
+									.Text(LOCTEXT("SColorBlockLabel", "SColorBlock"))
 							]
 
-							+ SVerticalBox::Slot()
-							.AutoHeight()
+						+ SGridPanel::Slot(1, 7)
+							.Padding(0.0f, 5.0f)
 							[
-								// second combo box
-								SAssignNew(SecondComboBox, SComboBox<TSharedPtr<FString> >)
-									.OptionsSource(&SecondComboBoxOptions)
-									.OnSelectionChanged(this, &SWidgetGallery::HandleSecondComboBoxSelectionChanged)
-									.OnGenerateWidget(this, &SWidgetGallery::HandleComboBoxGenerateWidget)
-									[
-										SNew(STextBlock)
-											.Text(this, &SWidgetGallery::HandleSecondComboBoxText)
-									]
+								SNew(SColorBlock)
+									.Color(FLinearColor(1.0f, 0.0f, 0.0f))
 							]
-						]
 
-					// SComboButton
-					+ SGridPanel::Slot(0, 9)
-						[
-							SNew(STextBlock)
-								.Text(LOCTEXT("SComboButtonLabel", "SComboButton"))
-						]
+						// SComboBox
+						+ SGridPanel::Slot(0, 8)
+							[
+								SNew(STextBlock)
+									.Text(LOCTEXT("SComboBoxLabel", "SComboBox"))
+							]
 
-					+ SGridPanel::Slot(1, 9)
-						.HAlign(HAlign_Left)
-						.Padding(0.0f, 5.0f)
-						[
-							SNew(SFxWidget).IgnoreClipping(false)
+						+ SGridPanel::Slot(1, 8)
+							.HAlign(HAlign_Left)
+							.Padding(0.0f, 5.0f)
 							[
-								SNew(SComboButton)
-								.ButtonContent()
+								SNew(SVerticalBox)
+
+								+ SVerticalBox::Slot()
 								[
-									SNew(STextBlock)
-										.Text(LOCTEXT("ComboButtonLabel", "Combo Button"))
+									SNew(SCheckBox)
+									.IsChecked(this, &SWidgetGallery::IsCreateNewWindowChecked)
+									.OnCheckStateChanged( this, &SWidgetGallery::CreateNewWindowToggled )
+									[
+										SNew(STextBlock).Text(LOCTEXT("CreateNewWindowCheckbox", "Create new window for popups."))
+									]
 								]
-								.MenuContent()
-								[
-									SNew(SVerticalBox)
 
-									+ SVerticalBox::Slot()
-										.AutoHeight()
+								+ SVerticalBox::Slot()
+								.AutoHeight()
+								[
+									// selector combo box
+									SAssignNew(SelectorComboBox, SComboBox<TSharedPtr<FString> >)
+										.OptionsSource(&SelectorComboBoxOptions)
+										.OnSelectionChanged(this, &SWidgetGallery::HandleSelectorComboBoxSelectionChanged)
+										.OnGenerateWidget(this, &SWidgetGallery::HandleComboBoxGenerateWidget)
 										[
 											SNew(STextBlock)
-												.Text(LOCTEXT("ComboButtonItemLabel01", "Combo\n     Button\n  menu\n       content supports"))
-										]
-
-									+ SVerticalBox::Slot()
-										.AutoHeight()
-										.HAlign(HAlign_Center)
-										[
-											SNew(SButton)
-												[ 
-													SNew(STextBlock)
-														.Text(LOCTEXT("ComboButtonItemLabel02", "arbitrary"))
-												]
+												.Text(this, &SWidgetGallery::HandleSelectorComboBoxText)
 										]
+								]
 
-									+ SVerticalBox::Slot()
-										.AutoHeight()
+								+ SVerticalBox::Slot()
+								.AutoHeight()
+								[
+									// second combo box
+									SAssignNew(SecondComboBox, SComboBox<TSharedPtr<FString> >)
+										.OptionsSource(&SecondComboBoxOptions)
+										.OnSelectionChanged(this, &SWidgetGallery::HandleSecondComboBoxSelectionChanged)
+										.OnGenerateWidget(this, &SWidgetGallery::HandleComboBoxGenerateWidget)
 										[
 											SNew(STextBlock)
-												.Text(LOCTEXT("ComboButtonItemLabel03", "widgets"))
+												.Text(this, &SWidgetGallery::HandleSecondComboBoxText)
 										]
 								]
 							]
-						]
-
-					// SEditableText
-					+ SGridPanel::Slot(0, 10)
-						[
-							SNew(STextBlock)
-								.Text(LOCTEXT("SEditableTextLabel", "SEditableText"))
-						]
-
-					+ SGridPanel::Slot(1, 10)
-						.HAlign(HAlign_Left)
-						.Padding(0.0f, 5.0f)
-						[
-							SNew(SEditableText)
-								.HintText(LOCTEXT("SEditableTextHint", "This is editable text"))
-						]
-
-					// SEditableTextBox
-					+ SGridPanel::Slot(0, 11)
-						[
-							SNew(STextBlock)
-								.Text(LOCTEXT("SEditableTextBoxLabel", "SEditableTextBox"))
-						]
-
-					+ SGridPanel::Slot(1, 11)
-						.HAlign(HAlign_Left)
-						.Padding(0.0f, 5.0f)
-						[
-							SNew(SEditableTextBox)
-								.HintText(LOCTEXT("SEditableTextBoxHint", "This is an editable text box"))
-						]
-
-					// SHeader
-					+ SGridPanel::Slot(0, 12)
-						[
-							SNew(STextBlock)
-								.Text(LOCTEXT("SHeaderLabel", "SHeader"))
-						]
 
-					+ SGridPanel::Slot(1, 12)
-						.Padding(0.0f, 5.0f)
-						[
-							SNew(SHeader)
-								.Content()
-								[
-									SNew(STextBlock)
-										.Text(LOCTEXT("HeaderContentLabel", "Header Content"))
-								]
-						]
-
-					// SHyperlink
-					+ SGridPanel::Slot(0, 13)
-						[
-							SNew(STextBlock)
-								.Text(LOCTEXT("SHyperlinkLabel", "SHyperlink"))
-						]
-
-					+ SGridPanel::Slot(1, 13)
-						.HAlign(HAlign_Left)
-						.Padding(0.0f, 5.0f)
-						[
-							SNew(SHyperlink)
-								.Text(LOCTEXT("SHyperlinkText", "Hyperlink"))
-			
-						]
-
-					// SImage
-					+ SGridPanel::Slot(0, 14)
-						[
-							SNew(STextBlock)
-								.Text(LOCTEXT("SImageLabel", "SImage"))
-						]
-
-					+ SGridPanel::Slot(1, 14)
-						.Padding(0.0f, 5.0f)
-						[
-							SNew(SImage)
-								.Image(FTestStyle::Get().GetBrush(TEXT("NewLevelBlank")))			
-						]
-/*
-					// SNumericEntryBox
-					+ SGridPanel::Slot(0, 15)
-						[
-							SNew(STextBlock)
-								.Text(LOCTEXT("SNumericEntryBoxLabel", "SNumericEntryBox"))
-						]
-
-					+ SGridPanel::Slot(1, 15)
-						[
-							New(SNumericEntryBox<float>)
-						]
-*/
-					// SProgressBar
-					+ SGridPanel::Slot(0, 16)
-						[
-							SNew(STextBlock)
-								.Text(LOCTEXT("SProgressBarLabel", "SProgressBar"))
-						]
+						// SComboButton
+						+ SGridPanel::Slot(0, 9)
+							[
+								SNew(STextBlock)
+									.Text(LOCTEXT("SComboButtonLabel", "SComboButton"))
+							]
 
-					+ SGridPanel::Slot(1, 16)
-						.Padding(0.0f, 5.0f)
-						[
-							SNew(SBox)
-								.WidthOverride(150.0f)
+						+ SGridPanel::Slot(1, 9)
+							.HAlign(HAlign_Left)
+							.Padding(0.0f, 5.0f)
+							[
+								SNew(SFxWidget).IgnoreClipping(false)
 								[
-									SNew(SProgressBar)
-										.Percent(this, &SWidgetGallery::HandleProgressBarPercent)
+									SNew(SComboButton)
+									.ButtonContent()
+									[
+										SNew(STextBlock)
+											.Text(LOCTEXT("ComboButtonLabel", "Combo Button"))
+									]
+									.MenuContent()
+									[
+										SNew(SVerticalBox)
+
+										+ SVerticalBox::Slot()
+											.AutoHeight()
+											[
+												SNew(STextBlock)
+													.Text(LOCTEXT("ComboButtonItemLabel01", "Combo\n     Button\n  menu\n       content supports"))
+											]
+
+										+ SVerticalBox::Slot()
+											.AutoHeight()
+											.HAlign(HAlign_Center)
+											[
+												SNew(SButton)
+													[ 
+														SNew(STextBlock)
+															.Text(LOCTEXT("ComboButtonItemLabel02", "arbitrary"))
+													]
+											]
+
+										+ SVerticalBox::Slot()
+											.AutoHeight()
+											[
+												SNew(STextBlock)
+													.Text(LOCTEXT("ComboButtonItemLabel03", "widgets"))
+											]
+									]
 								]
-						]
-/*
-					// SRotatorInputBox
-					+ SGridPanel::Slot(0, 17)
-						[
-							SNew(STextBlock)
-								.Text(LOCTEXT("SRotatorInputBoxLabel", "SRotatorInputBox"))
-						]
-
-					+ SGridPanel::Slot(1, 17)
-						.Padding(0.0f, 5.0f)
-						[
-							SNew(SRotatorInputBox)
-								.Roll(0.5f)
-								.Pitch(0.0f)
-								.Yaw(1.0f)			
-						]
-*/
-					// SSearchBox
-					+ SGridPanel::Slot(0, 18)
-						[
-							SNew(STextBlock)
-								.Text(LOCTEXT("SSearchBoxLabel", "SSearchBox"))
-						]
+							]
 
-					+ SGridPanel::Slot(1, 18)
-						.Padding(0.0f, 5.0f)
-						[
-							SNew(SVerticalBox)
+						// SEditableText
+						+ SGridPanel::Slot(0, 10)
+							[
+								SNew(STextBlock)
+									.Text(LOCTEXT("SEditableTextLabel", "SEditableText"))
+							]
 
-							+ SVerticalBox::Slot()
-								.AutoHeight()
-								.HAlign(HAlign_Left)
-								[
-									SNew(SSearchBox)
-								]
+						+ SGridPanel::Slot(1, 10)
+							.HAlign(HAlign_Left)
+							.Padding(0.0f, 5.0f)
+							[
+								SNew(SEditableText)
+									.HintText(LOCTEXT("SEditableTextHint", "This is editable text"))
+							]
 
-							+ SVerticalBox::Slot()
-								.AutoHeight()
-								.HAlign(HAlign_Left)
-								.Padding(0.0f, 4.0f, 0.0f, 0.0f)
-								[
-									SNew(SSearchBox)
-										.DelayChangeNotificationsWhileTyping(true)
-										.OnSearch(this, &SWidgetGallery::HandleSearchBoxSearch)
-								]			
-						]
+						// SEditableTextBox
+						+ SGridPanel::Slot(0, 11)
+							[
+								SNew(STextBlock)
+									.Text(LOCTEXT("SEditableTextBoxLabel", "SEditableTextBox"))
+							]
 
-					// SSeparator
-					+ SGridPanel::Slot(0, 19)
-						[
-							SNew(STextBlock)
-								.Text(LOCTEXT("SSeparatorLabel", "SSeparator"))
-						]
+						+ SGridPanel::Slot(1, 11)
+							.HAlign(HAlign_Left)
+							.Padding(0.0f, 5.0f)
+							[
+								SNew(SEditableTextBox)
+									.HintText(LOCTEXT("SEditableTextBoxHint", "This is an editable text box"))
+							]
 
-					+ SGridPanel::Slot(1, 19)
-						.HAlign(HAlign_Left)
-						.Padding(0.0f, 5.0f)
-						[
-							SNew(SBox)
-								.HeightOverride(100.0f)
-								.WidthOverride(150.0f)
-								[
-									SNew(SHorizontalBox)
+						// SHeader
+						+ SGridPanel::Slot(0, 12)
+							[
+								SNew(STextBlock)
+									.Text(LOCTEXT("SHeaderLabel", "SHeader"))
+							]
 
-									+ SHorizontalBox::Slot()
-										.VAlign(VAlign_Center)
-										.FillWidth(0.75f)
-										[
-											SNew(SSeparator)
-												.Orientation(Orient_Horizontal)
-										]
+						+ SGridPanel::Slot(1, 12)
+							.Padding(0.0f, 5.0f)
+							[
+								SNew(SHeader)
+									.Content()
+									[
+										SNew(STextBlock)
+											.Text(LOCTEXT("HeaderContentLabel", "Header Content"))
+									]
+							]
 
-									+ SHorizontalBox::Slot()
-										.HAlign(HAlign_Center)
-										.FillWidth(0.25f)
-										[
-											SNew(SSeparator)
-												.Orientation(Orient_Vertical)
-										]
-								]			
-						]
+						// SHyperlink
+						+ SGridPanel::Slot(0, 13)
+							[
+								SNew(STextBlock)
+									.Text(LOCTEXT("SHyperlinkLabel", "SHyperlink"))
+							]
 
-					// SSlider
-					+ SGridPanel::Slot(0, 20)
-						[
-							SNew(STextBlock)
-								.Text(LOCTEXT("SSliderLabel", "SSlider"))
-						]
+						+ SGridPanel::Slot(1, 13)
+							.HAlign(HAlign_Left)
+							.Padding(0.0f, 5.0f)
+							[
+								SNew(SHyperlink)
+									.Text(LOCTEXT("SHyperlinkText", "Hyperlink"))
+			
+							]
 
-					+ SGridPanel::Slot(1, 20)
-						.HAlign(HAlign_Left)
-						.Padding(0.0f, 5.0f)
-						[
-							SNew(SBox)
-								.HeightOverride(100.0f)
-								.WidthOverride(150.0f)
-								[
-									SNew(SHorizontalBox)
+						// SImage
+						+ SGridPanel::Slot(0, 14)
+							[
+								SNew(STextBlock)
+									.Text(LOCTEXT("SImageLabel", "SImage"))
+							]
 
-									+ SHorizontalBox::Slot()
-										.VAlign(VAlign_Center)
-										.FillWidth(0.75f)
-										[
-											SNew(SSlider)
-												.Orientation(Orient_Horizontal)
-												.Value(0.5f)
-										]
+						+ SGridPanel::Slot(1, 14)
+							.Padding(0.0f, 5.0f)
+							[
+								SNew(SImage)
+									.Image(FTestStyle::Get().GetBrush(TEXT("NewLevelBlank")))			
+							]
+	/*
+						// SNumericEntryBox
+						+ SGridPanel::Slot(0, 15)
+							[
+								SNew(STextBlock)
+									.Text(LOCTEXT("SNumericEntryBoxLabel", "SNumericEntryBox"))
+							]
 
-									+ SHorizontalBox::Slot()
-										.HAlign(HAlign_Center)
-										.FillWidth(0.25f)
-										[
-											SNew(SSlider)
-												.Orientation(Orient_Vertical)
-												.Value(0.5f)
-										]
-								]			
-						]
+						+ SGridPanel::Slot(1, 15)
+							[
+								New(SNumericEntryBox<float>)
+							]
+	*/
+						// SProgressBar
+						+ SGridPanel::Slot(0, 16)
+							[
+								SNew(STextBlock)
+									.Text(LOCTEXT("SProgressBarLabel", "SProgressBar"))
+							]
 
-					// SSlider (no indentation)
-					+ SGridPanel::Slot(0, 21)
-						[
-							SNew(STextBlock)
-								.Text(LOCTEXT("SSliderNoIndentLabel", "SSlider (no indentation)"))
-						]
+						+ SGridPanel::Slot(1, 16)
+							.Padding(0.0f, 5.0f)
+							[
+								SNew(SBox)
+									.WidthOverride(150.0f)
+									[
+										SNew(SProgressBar)
+											.Percent(this, &SWidgetGallery::HandleProgressBarPercent)
+									]
+							]
+	/*
+						// SRotatorInputBox
+						+ SGridPanel::Slot(0, 17)
+							[
+								SNew(STextBlock)
+									.Text(LOCTEXT("SRotatorInputBoxLabel", "SRotatorInputBox"))
+							]
 
-					+ SGridPanel::Slot(1, 21)
-						.Padding(0.0f, 5.0f)
-						[
-							SNew(SBox)
-								.HeightOverride(100.0f)
-								.WidthOverride(150.0f)
-								[
-									SNew(SHorizontalBox)
+						+ SGridPanel::Slot(1, 17)
+							.Padding(0.0f, 5.0f)
+							[
+								SNew(SRotatorInputBox)
+									.Roll(0.5f)
+									.Pitch(0.0f)
+									.Yaw(1.0f)			
+							]
+	*/
+						// SSearchBox
+						+ SGridPanel::Slot(0, 18)
+							[
+								SNew(STextBlock)
+									.Text(LOCTEXT("SSearchBoxLabel", "SSearchBox"))
+							]
 
-									+ SHorizontalBox::Slot()
-										.VAlign(VAlign_Center)
-										.FillWidth(0.75f)
-										[
-											SNew(SSlider)
-												.IndentHandle(false)
-												.Orientation(Orient_Horizontal)
-												.Value(0.5f)
-										]
+						+ SGridPanel::Slot(1, 18)
+							.Padding(0.0f, 5.0f)
+							[
+								SNew(SVerticalBox)
 
-									+ SHorizontalBox::Slot()
-										.HAlign(HAlign_Center)
-										.FillWidth(0.25f)
-										[
-											SNew(SSlider)
-												.IndentHandle(false)
-												.Orientation(Orient_Vertical)
-												.Value(0.5f)
-										]
-								]			
-						]
+								+ SVerticalBox::Slot()
+									.AutoHeight()
+									.HAlign(HAlign_Left)
+									[
+										SNew(SSearchBox)
+									]
 
-					// SSpacer
-					+ SGridPanel::Slot(0, 22)
-						[
-							SNew(STextBlock)
-								.Text(LOCTEXT("SSpacerLabel", "SSpacer"))
-						]
+								+ SVerticalBox::Slot()
+									.AutoHeight()
+									.HAlign(HAlign_Left)
+									.Padding(0.0f, 4.0f, 0.0f, 0.0f)
+									[
+										SNew(SSearchBox)
+											.DelayChangeNotificationsWhileTyping(true)
+											.OnSearch(this, &SWidgetGallery::HandleSearchBoxSearch)
+									]			
+							]
 
-					+ SGridPanel::Slot(1, 22)
-						.Padding(0.0f, 5.0f)
-						[
-							SNew(SSpacer)
-								.Size(FVector2D(100, 100))		
-						]
+						// SSeparator
+						+ SGridPanel::Slot(0, 19)
+							[
+								SNew(STextBlock)
+									.Text(LOCTEXT("SSeparatorLabel", "SSeparator"))
+							]
 
-					// SSpinningImage
-					+ SGridPanel::Slot(0, 23)
-						[
-							SNew(STextBlock)
-								.Text(LOCTEXT("SSpinningImageLabel", "SSpinningImage"))
-						]
+						+ SGridPanel::Slot(1, 19)
+							.HAlign(HAlign_Left)
+							.Padding(0.0f, 5.0f)
+							[
+								SNew(SBox)
+									.HeightOverride(100.0f)
+									.WidthOverride(150.0f)
+									[
+										SNew(SHorizontalBox)
+
+										+ SHorizontalBox::Slot()
+											.VAlign(VAlign_Center)
+											.FillWidth(0.75f)
+											[
+												SNew(SSeparator)
+													.Orientation(Orient_Horizontal)
+											]
+
+										+ SHorizontalBox::Slot()
+											.HAlign(HAlign_Center)
+											.FillWidth(0.25f)
+											[
+												SNew(SSeparator)
+													.Orientation(Orient_Vertical)
+											]
+									]			
+							]
 
-					+ SGridPanel::Slot(1, 23)
-						.HAlign(HAlign_Left)
-						.Padding(0.0f, 5.0f)
-						[
-							SNew(SSpinningImage)
-								.Image(FTestStyle::Get().GetBrush("TestRotation16px"))
-						]
+						// SSlider
+						+ SGridPanel::Slot(0, 20)
+							[
+								SNew(STextBlock)
+									.Text(LOCTEXT("SSliderLabel", "SSlider"))
+							]
 
-					// SSpinBox
-					+ SGridPanel::Slot(0, 24)
-						[
-							SNew(STextBlock)
-								.Text(LOCTEXT("SSpinBoxLabel", "SSpinBox"))
-						]
+						+ SGridPanel::Slot(1, 20)
+							.HAlign(HAlign_Left)
+							.Padding(0.0f, 5.0f)
+							[
+								SNew(SBox)
+									.HeightOverride(100.0f)
+									.WidthOverride(150.0f)
+									[
+										SNew(SHorizontalBox)
+
+										+ SHorizontalBox::Slot()
+											.VAlign(VAlign_Center)
+											.FillWidth(0.75f)
+											[
+												SNew(SSlider)
+													.Orientation(Orient_Horizontal)
+													.Value(0.5f)
+											]
+
+										+ SHorizontalBox::Slot()
+											.HAlign(HAlign_Center)
+											.FillWidth(0.25f)
+											[
+												SNew(SSlider)
+													.Orientation(Orient_Vertical)
+													.Value(0.5f)
+											]
+									]			
+							]
 
-					+ SGridPanel::Slot(1, 24)
-						.HAlign(HAlign_Left)
-						.Padding(0.0f, 5.0f)
-						[
-							SNew(SSpinBox<float>)
-								.MinValue(-1000.0f)
-								.MaxValue(1000.0f)
-								.MinSliderValue(TAttribute< TOptional<float> >(-500.0f))
-								.MaxSliderValue(TAttribute< TOptional<float> >(500.0f))
-								.Delta(0.5f)
-						]
+						// SSlider (no indentation)
+						+ SGridPanel::Slot(0, 21)
+							[
+								SNew(STextBlock)
+									.Text(LOCTEXT("SSliderNoIndentLabel", "SSlider (no indentation)"))
+							]
 
-					// STextBlock
-					+ SGridPanel::Slot(0, 25)
-						[
-							SNew(STextBlock)
-								.Text(LOCTEXT("STextBlockLabel", "STextBlock"))
-						]
+						+ SGridPanel::Slot(1, 21)
+							.Padding(0.0f, 5.0f)
+							[
+								SNew(SBox)
+									.HeightOverride(100.0f)
+									.WidthOverride(150.0f)
+									[
+										SNew(SHorizontalBox)
+
+										+ SHorizontalBox::Slot()
+											.VAlign(VAlign_Center)
+											.FillWidth(0.75f)
+											[
+												SNew(SSlider)
+													.IndentHandle(false)
+													.Orientation(Orient_Horizontal)
+													.Value(0.5f)
+											]
+
+										+ SHorizontalBox::Slot()
+											.HAlign(HAlign_Center)
+											.FillWidth(0.25f)
+											[
+												SNew(SSlider)
+													.IndentHandle(false)
+													.Orientation(Orient_Vertical)
+													.Value(0.5f)
+											]
+									]			
+							]
 
-					+ SGridPanel::Slot(1, 25)
-						.Padding(0.0f, 5.0f)
-						[
-							SNew(SVerticalBox)
-							
-							+ SVerticalBox::Slot()
-							.AutoHeight()
+						// SSpacer
+						+ SGridPanel::Slot(0, 22)
 							[
 								SNew(STextBlock)
-									.Text(LOCTEXT("STextBlockLeftExampleLabel", "This is a left-justified text block"))
-									.Justification(ETextJustify::Left)
+									.Text(LOCTEXT("SSpacerLabel", "SSpacer"))
 							]
 
-							+ SVerticalBox::Slot()
-							.AutoHeight()
+						+ SGridPanel::Slot(1, 22)
+							.Padding(0.0f, 5.0f)
 							[
-								SNew(STextBlock)
-									.Text(LOCTEXT("STextBlockCenterExampleLabel", "This is a center-justified text block"))
-									.Justification(ETextJustify::Center)
+								SNew(SSpacer)
+									.Size(FVector2D(100, 100))		
 							]
 
-							+ SVerticalBox::Slot()
-							.AutoHeight()
+						// SSpinningImage
+						+ SGridPanel::Slot(0, 23)
 							[
 								SNew(STextBlock)
-									.Text(LOCTEXT("STextBlockRightExampleLabel", "This is a right-justified text block"))
-									.Justification(ETextJustify::Right)
+									.Text(LOCTEXT("SSpinningImageLabel", "SSpinningImage"))
 							]
-						]
-
-					// STextComboBox
-					+ SGridPanel::Slot(0, 26)
-						[
-							SNew(STextBlock)
-								.Text(LOCTEXT("STextComboBoxLabel", "STextComboBox"))
-						]
 
-					+ SGridPanel::Slot(1, 26)
-						.Padding(0.0f, 5.0f)
-						.HAlign(HAlign_Left)
-						.VAlign(VAlign_Center)
-						[
-							SAssignNew(TextComboBox, STextComboBox)
-								.OptionsSource(&TextComboBoxOptions)
-								.OnSelectionChanged(this, &SWidgetGallery::HandleTextComboBoxSelectionChanged)
-								.OnGetTextLabelForItem(this, &SWidgetGallery::HandleTextComboBoxGetTextLabelForItem)
-								.InitiallySelectedItem(TextComboBoxSelectedItem)			
-						]
+						+ SGridPanel::Slot(1, 23)
+							.HAlign(HAlign_Left)
+							.Padding(0.0f, 5.0f)
+							[
+								SNew(SSpinningImage)
+									.Image(FTestStyle::Get().GetBrush("TestRotation16px"))
+							]
 
-					// STextComboPopup
-					+ SGridPanel::Slot(0, 27)
-						[
-							SNew(STextBlock)
-								.Text(LOCTEXT("STextComboPopupLabel", "STextComboPopup"))
-						]
+						// SSpinBox
+						+ SGridPanel::Slot(0, 24)
+							[
+								SNew(STextBlock)
+									.Text(LOCTEXT("SSpinBoxLabel", "SSpinBox"))
+							]
 
-					+ SGridPanel::Slot(1, 27)
-						.Padding(0.0f, 5.0f)
-						.HAlign(HAlign_Left)
-						.VAlign(VAlign_Center)
-						[
-							SNew(SButton)
-								.Text(LOCTEXT("ButtonTestLabel", "Test"))
-								.OnClicked(this, &SWidgetGallery::HandleTextComboPopupClicked)
-						]
-		
-					// SThrobber
-					+ SGridPanel::Slot(0, 28)
-						[
-							SNew(STextBlock)
-								.Text(LOCTEXT("SThrobberLabel", "SThrobber"))
-						]
+						+ SGridPanel::Slot(1, 24)
+							.HAlign(HAlign_Left)
+							.Padding(0.0f, 5.0f)
+							[
+								SNew(SSpinBox<float>)
+									.MinValue(-1000.0f)
+									.MaxValue(1000.0f)
+									.MinSliderValue(TAttribute< TOptional<float> >(-500.0f))
+									.MaxSliderValue(TAttribute< TOptional<float> >(500.0f))
+									.Delta(0.5f)
+							]
 
-					+ SGridPanel::Slot(1, 28)
-						.Padding(0.0f, 5.0f)
-						[
-							SNew(SVerticalBox)
+						// STextBlock
+						+ SGridPanel::Slot(0, 25)
+							[
+								SNew(STextBlock)
+									.Text(LOCTEXT("STextBlockLabel", "STextBlock"))
+							]
 
-							+ SVerticalBox::Slot()
+						+ SGridPanel::Slot(1, 25)
+							.Padding(0.0f, 5.0f)
+							[
+								SNew(SVerticalBox)
+							
+								+ SVerticalBox::Slot()
 								.AutoHeight()
 								[
-									SNew(SThrobber)
-										.Animate(SThrobber::Horizontal)
+									SNew(STextBlock)
+										.Text(LOCTEXT("STextBlockLeftExampleLabel", "This is a left-justified text block"))
+										.Justification(ETextJustify::Left)
 								]
 
-							+ SVerticalBox::Slot()
+								+ SVerticalBox::Slot()
 								.AutoHeight()
 								[
-									SNew(SThrobber)
-										.Animate(SThrobber::Opacity)
+									SNew(STextBlock)
+										.Text(LOCTEXT("STextBlockCenterExampleLabel", "This is a center-justified text block"))
+										.Justification(ETextJustify::Center)
 								]
 
-							+ SVerticalBox::Slot()
+								+ SVerticalBox::Slot()
 								.AutoHeight()
 								[
-									SNew(SThrobber)
-										.Animate(SThrobber::VerticalAndOpacity) .NumPieces(5)
+									SNew(STextBlock)
+										.Text(LOCTEXT("STextBlockRightExampleLabel", "This is a right-justified text block"))
+										.Justification(ETextJustify::Right)
 								]
-						]
-/*
-					// SVectorInputBox
-					+ SGridPanel::Slot(0, 29)
-						[
-							SNew(STextBlock)
-								.Text(LOCTEXT("SVectorInputBoxLabel", "SVectorInputBox"))
-						]
+							]
 
-					+ SGridPanel::Slot(1, 29)
-						.Padding(0.0f, 5.0f)
-						[
-							SNew(SVectorInputBox)
-								.X(-5)
-								.Y(0)
-								.Z(5)			
-						]
-*/
-					// SVolumeControl
-					+ SGridPanel::Slot(0, 30)
+						// STextComboBox
+						+ SGridPanel::Slot(0, 26)
+							[
+								SNew(STextBlock)
+									.Text(LOCTEXT("STextComboBoxLabel", "STextComboBox"))
+							]
+
+						+ SGridPanel::Slot(1, 26)
+							.Padding(0.0f, 5.0f)
+							.HAlign(HAlign_Left)
+							.VAlign(VAlign_Center)
+							[
+								SAssignNew(TextComboBox, STextComboBox)
+									.OptionsSource(&TextComboBoxOptions)
+									.OnSelectionChanged(this, &SWidgetGallery::HandleTextComboBoxSelectionChanged)
+									.OnGetTextLabelForItem(this, &SWidgetGallery::HandleTextComboBoxGetTextLabelForItem)
+									.InitiallySelectedItem(TextComboBoxSelectedItem)			
+							]
+
+						// STextComboPopup
+						+ SGridPanel::Slot(0, 27)
+							[
+								SNew(STextBlock)
+									.Text(LOCTEXT("STextComboPopupLabel", "STextComboPopup"))
+							]
+
+						+ SGridPanel::Slot(1, 27)
+							.Padding(0.0f, 5.0f)
+							.HAlign(HAlign_Left)
+							.VAlign(VAlign_Center)
+							[
+								SNew(SButton)
+									.Text(LOCTEXT("ButtonTestLabel", "Test"))
+									.OnClicked(this, &SWidgetGallery::HandleTextComboPopupClicked)
+							]
+		
+						// SThrobber
+						+ SGridPanel::Slot(0, 28)
+							[
+								SNew(STextBlock)
+									.Text(LOCTEXT("SThrobberLabel", "SThrobber"))
+							]
+
+						+ SGridPanel::Slot(1, 28)
+							.Padding(0.0f, 5.0f)
+							[
+								SNew(SVerticalBox)
+
+								+ SVerticalBox::Slot()
+									.AutoHeight()
+									[
+										SNew(SThrobber)
+											.Animate(SThrobber::Horizontal)
+									]
+
+								+ SVerticalBox::Slot()
+									.AutoHeight()
+									[
+										SNew(SThrobber)
+											.Animate(SThrobber::Opacity)
+									]
+
+								+ SVerticalBox::Slot()
+									.AutoHeight()
+									[
+										SNew(SThrobber)
+											.Animate(SThrobber::VerticalAndOpacity) .NumPieces(5)
+									]
+							]
+	/*
+						// SVectorInputBox
+						+ SGridPanel::Slot(0, 29)
+							[
+								SNew(STextBlock)
+									.Text(LOCTEXT("SVectorInputBoxLabel", "SVectorInputBox"))
+							]
+
+						+ SGridPanel::Slot(1, 29)
+							.Padding(0.0f, 5.0f)
+							[
+								SNew(SVectorInputBox)
+									.X(-5)
+									.Y(0)
+									.Z(5)			
+							]
+	*/
+						// SVolumeControl
+						+ SGridPanel::Slot(0, 30)
+							[
+								SNew(STextBlock)
+									.Text(LOCTEXT("SVolumeControlLabel", "SVolumeControl"))
+							]
+
+						+ SGridPanel::Slot(1, 30)
+							.HAlign(HAlign_Left)
+							.Padding(0.0f, 5.0f)
+							[
+								SNew(SBox)
+									.WidthOverride(150.0f)
+									[
+										SNew(SVolumeControl)
+											.Volume(0.6f)
+									]
+							]
+
+						+ SGridPanel::Slot(0, 31)
 						[
 							SNew(STextBlock)
-								.Text(LOCTEXT("SVolumeControlLabel", "SVolumeControl"))
+								.Text(LOCTEXT("SUserWidgetLabel", "SUserWidgetExample"))
 						]
 
-					+ SGridPanel::Slot(1, 30)
+						+ SGridPanel::Slot(1, 31)
 						.HAlign(HAlign_Left)
 						.Padding(0.0f, 5.0f)
 						[
-							SNew(SBox)
-								.WidthOverride(150.0f)
-								[
-									SNew(SVolumeControl)
-										.Volume(0.6f)
-								]
+							SNew(SUserWidgetExample)
+							.Title( LOCTEXT("SUserWidgetTest", "UserWidgetTest.cpp") )
 						]
-
-					+ SGridPanel::Slot(0, 31)
-					[
-						SNew(STextBlock)
-							.Text(LOCTEXT("SUserWidgetLabel", "SUserWidgetExample"))
-					]
-
-					+ SGridPanel::Slot(1, 31)
-					.HAlign(HAlign_Left)
-					.Padding(0.0f, 5.0f)
-					[
-						SNew(SUserWidgetExample)
-						.Title( LOCTEXT("SUserWidgetTest", "UserWidgetTest.cpp") )
 					]
-				]
+			]
 		];
 
 		// initialize SBreadcrumbTrail
diff --git a/Engine/Source/Runtime/AppFramework/Private/Widgets/Testing/STestSuite.cpp b/Engine/Source/Runtime/AppFramework/Private/Widgets/Testing/STestSuite.cpp
--- a/Engine/Source/Runtime/AppFramework/Private/Widgets/Testing/STestSuite.cpp
+++ b/Engine/Source/Runtime/AppFramework/Private/Widgets/Testing/STestSuite.cpp
@@ -5115,18 +5115,6 @@ TSharedRef<SDockTab> SpawnTab(const FSpawnTabArgs& Args, FName TabIdentifier)
 			]
 		];
 	}
-	else if (TabIdentifier == FName(TEXT("WidgetGalleryTab")))
-	{
-		return SNew(SDockTab)
-			. Label( LOCTEXT("WidgetGalleryTab", "Widget Gallery") )
-			. ToolTipText( LOCTEXT( "WidgetGalleryTabTextToolTip", "Switch to the widget gallery." ) )
-		[
-			SNew(SScissorRectBox)
-			[
-				MakeWidgetGallery()
-			]
-		];
-	}
 	else if (TabIdentifier == FName(TEXT("ColorPickerTestTab")))
 	{
 		return SNew(SDockTab)
@@ -5185,6 +5173,21 @@ TSharedRef<SDockTab> SpawnTab(const FSpawnTabArgs& Args, FName TabIdentifier)
 		return SNew(SDockTab);
 	}
 }
+
+TSharedRef<SDockTab> SpawnResponsiveGrid(const FSpawnTabArgs& Args)
+{
+	TSharedRef<SDockTab> ResponsiveGridTab =
+		SNew(SDockTab)
+		.Label(LOCTEXT("ResponsiveGridTabLabel", "Responsive Grid"))
+		.ToolTipText(LOCTEXT("ResponsiveGridTabToolTip", ""));
+
+	ResponsiveGridTab->SetContent
+		(
+		SNew(SResponsiveGridPanelTestWidget)
+		);
+
+	return ResponsiveGridTab;
+}
 END_SLATE_FUNCTION_BUILD_OPTIMIZATION
 
 
@@ -5227,7 +5230,7 @@ TSharedRef<SDockTab> SpawnTestSuite1( const FSpawnTabArgs& Args )
 				FTabManager::NewStack()
 				->SetSizeCoefficient(0.25f)
 				->AddTab("MultiBoxTestTab", ETabState::OpenedTab)
-				->AddTab("WidgetGalleryTab", ETabState::OpenedTab)
+				->AddTab("ResponsiveGrid", ETabState::OpenedTab)
 				// The DocTest tab will be closed by default.
 				->AddTab("DocTest", ETabState::ClosedTab)
 			)
@@ -5301,9 +5304,8 @@ TSharedRef<SDockTab> SpawnTestSuite1( const FSpawnTabArgs& Args )
 		.SetDisplayName( NSLOCTEXT("TestSuite1", "LayoutRoundingTab", "Layout Rounding") )
 		.SetGroup(TestSuiteMenu::SuiteTabs);
 
-	TestSuite1TabManager->RegisterTabSpawner( "WidgetGalleryTab", FOnSpawnTab::CreateStatic( &SpawnTab, FName("WidgetGalleryTab") ) )
-		.SetDisplayName( NSLOCTEXT("TestSuite1", "WidgetGalleryTab", "Widget Gallery") )
-		.SetGroup(TestSuiteMenu::NestedCategory);
+	TestSuite1TabManager->RegisterTabSpawner("ResponsiveGrid", FOnSpawnTab::CreateStatic(&SpawnResponsiveGrid))
+		.SetGroup(TestSuiteMenu::SuiteTabs);
 
 	TestSuite1TabManager->RegisterTabSpawner( "MultiBoxTestTab", FOnSpawnTab::CreateStatic( &SpawnTab, FName("MultiBoxTestTab") ) )
 		.SetDisplayName( NSLOCTEXT("TestSuite1", "MultiBoxTestTab", "MultiBox Test") )
@@ -5437,20 +5439,18 @@ TSharedRef<SDockTab> SpawnRenderTransformManipulator( const FSpawnTabArgs& Args
 	return RenderTransformManipulatorTab;
 }
 
-TSharedRef<SDockTab> SpawnResponsiveGrid(const FSpawnTabArgs& Args)
+TSharedRef<SDockTab> SpawnWidgetGallery(const FSpawnTabArgs& Args)
 {
-	TSharedRef<SDockTab> ResponsiveGridTab =
-		SNew(SDockTab)
-		.TabRole(ETabRole::MajorTab)
-		.Label(LOCTEXT("ResponsiveGridTabLabel", "Responsive Grid"))
-		.ToolTipText(LOCTEXT("ResponsiveGridTabToolTip", ""));
-
-	ResponsiveGridTab->SetContent
-	(
-		SNew(SResponsiveGridPanelTestWidget)
-	);
-
-	return ResponsiveGridTab;
+	return SNew(SDockTab)
+		.TabRole(ETabRole::NomadTab)
+		.Label(LOCTEXT("WidgetGalleryTab", "Widget Gallery"))
+		.ToolTipText(LOCTEXT("WidgetGalleryTabTextToolTip", "Switch to the widget gallery."))
+		[
+			SNew(SScissorRectBox)
+			[
+				MakeWidgetGallery()
+			]
+		];
 }
 
 void RestoreSlateTestSuite()
@@ -5460,8 +5460,10 @@ void RestoreSlateTestSuite()
 	FGlobalTabmanager::Get()->RegisterTabSpawner("TestSuite1", FOnSpawnTab::CreateStatic( &SpawnTestSuite1 ) );
 	FGlobalTabmanager::Get()->RegisterTabSpawner("TestSuite2", FOnSpawnTab::CreateStatic( &SpawnTestSuite2 ) );
 	FGlobalTabmanager::Get()->RegisterTabSpawner("RenderTransformManipulator", FOnSpawnTab::CreateStatic(&SpawnRenderTransformManipulator));
-	FGlobalTabmanager::Get()->RegisterTabSpawner("ResponsiveGrid", FOnSpawnTab::CreateStatic(&SpawnResponsiveGrid));
-
+	FGlobalTabmanager::Get()->RegisterNomadTabSpawner("WidgetGalleryTab", FOnSpawnTab::CreateStatic(&SpawnWidgetGallery))
+		.SetDisplayName(LOCTEXT("WidgetGalleryTab", "Widget Gallery"))
+		.SetGroup(TestSuiteMenu::MenuRoot);
+	
 	TSharedRef<FTabManager::FLayout> Layout = FTabManager::NewLayout( "SlateTestSuite_Layout" )
 	->AddArea
 	(
@@ -5473,7 +5475,7 @@ void RestoreSlateTestSuite()
 			->AddTab( "TestSuite2", ETabState::OpenedTab )
 			->AddTab( "TestSuite1", ETabState::OpenedTab )
 			->AddTab("RenderTransformManipulator", ETabState::OpenedTab)
-			->AddTab("ResponsiveGrid", ETabState::OpenedTab)
+			->AddTab("WidgetGalleryTab", ETabState::OpenedTab)
 		)		
 	)
 	#if PLATFORM_SUPPORTS_MULTIPLE_NATIVE_WINDOWS
