Slate : Added a simple test for synthetic mouse moves.

diff --git a/Engine/Source/Runtime/AppFramework/Private/Widgets/Testing/STestSuite.cpp b/Engine/Source/Runtime/AppFramework/Private/Widgets/Testing/STestSuite.cpp
--- a/Engine/Source/Runtime/AppFramework/Private/Widgets/Testing/STestSuite.cpp
+++ b/Engine/Source/Runtime/AppFramework/Private/Widgets/Testing/STestSuite.cpp
@@ -721,6 +721,8 @@ public:
 			Documents.Add( MakeShareable( new FDocumentInfo( LOCTEXT("Document04", "Document 4") ) ) );
 			Documents.Add( MakeShareable( new FDocumentInfo( LOCTEXT("Document05", "Document 5") ) ) );
 		}
+
+		bButtonOneVisible = true;
 		
 		this->ChildSlot
 		[
@@ -763,6 +765,23 @@ public:
 					]
 				]
 			]
+			+ SVerticalBox::Slot()
+			[
+				SNew(SHorizontalBox)
+				+ SHorizontalBox::Slot()
+				[
+					SNew(SButton)
+					.Visibility( this, &SDocumentsTest::GetButtonOneVisibility )
+					.Text(LOCTEXT("Button1", "Button One"))
+					.OnClicked(this, &SDocumentsTest::ToggleButtonOneVisibility)
+				]
+				+ SHorizontalBox::Slot()
+				[
+					SNew(SButton)
+					.Text(LOCTEXT("Button2", "Button Two"))
+					.OnClicked(this, &SDocumentsTest::ToggleButtonOneVisibility)
+				]
+			]
 		];
 	}
 
@@ -801,10 +820,23 @@ public:
 		return FReply::Handled();
 	}
 
+	EVisibility GetButtonOneVisibility() const
+	{
+		return bButtonOneVisible ? EVisibility::Visible : EVisibility::Collapsed;
+	}
+
+	FReply ToggleButtonOneVisibility()
+	{
+		bButtonOneVisible = !bButtonOneVisible;
+		return FReply::Handled();
+	}
+
 	FTabManager* TabManager;
 
 	TArray< TSharedRef<FDocumentInfo> > Documents;
 
+	bool bButtonOneVisible;
+
 };
 
 class SSplitterTest : public SCompoundWidget
