Fixed small bug in Color Picker where original color isn't correctly restored upon pressing Cancel.

diff --git a/Engine/Source/Runtime/AppFramework/Private/Widgets/Colors/SColorPicker.cpp b/Engine/Source/Runtime/AppFramework/Private/Widgets/Colors/SColorPicker.cpp
--- a/Engine/Source/Runtime/AppFramework/Private/Widgets/Colors/SColorPicker.cpp
+++ b/Engine/Source/Runtime/AppFramework/Private/Widgets/Colors/SColorPicker.cpp
@@ -648,7 +648,7 @@ void SColorPicker::DiscardColor()
 	if (OnColorPickerCancelled.IsBound())
 	{
 		// let the user decide what to do about cancel
-		OnColorPickerCancelled.Execute( OldColor );
+		OnColorPickerCancelled.Execute( OldColor.HSVToLinearRGB() );
 	}
 	else
 	{	
