[INTEGRATE] Change 2438554 by [REDACTED_PERSON][REDACTED_HANDLE] on 2015/02/09 14:45:00

	Editor - Exposing more customization options for the SWizard to allow a slew of new styling options for the buttons.

diff --git a/Engine/Source/Runtime/AppFramework/Private/Widgets/Workflow/SWizard.cpp b/Engine/Source/Runtime/AppFramework/Private/Widgets/Workflow/SWizard.cpp
--- a/Engine/Source/Runtime/AppFramework/Private/Widgets/Workflow/SWizard.cpp
+++ b/Engine/Source/Runtime/AppFramework/Private/Widgets/Workflow/SWizard.cpp
@@ -21,7 +21,6 @@ bool SWizard::CanShowPage( int32 PageIndex ) const
 	return false;
 }
 
-
 void SWizard::Construct( const FArguments& InArgs )
 {
 	DesiredSize = InArgs._DesiredSize.Get();
@@ -37,124 +36,133 @@ void SWizard::Construct( const FArguments& InArgs )
 		SNew(SVerticalBox)
 
 		+ SVerticalBox::Slot()
-			.FillHeight(1.0)			
+		.FillHeight(1.0)
+		[
+			SNew(SHorizontalBox)
+
+			+ SHorizontalBox::Slot()
+			.AutoWidth()
+			.Padding(0.0f, 0.0f, 20.0f, 0.0f)
+			[
+				SAssignNew(PageListBox, SVerticalBox)
+				.Visibility(InArgs._ShowPageList ? EVisibility::Visible : EVisibility::Collapsed)
+			]
+
+			+ SHorizontalBox::Slot()
+			.FillWidth(1.0f)
+			[
+				// widget switcher
+				SAssignNew(WidgetSwitcher, SWidgetSwitcher)
+			]
+		]
+
+		+ SVerticalBox::Slot()
+		.AutoHeight()
+		.Padding(0)
+		[
+			InArgs._PageFooter.Widget
+		]
+
+		+ SVerticalBox::Slot()
+		.AutoHeight()
+		.HAlign(HAlign_Right)
+		.Padding(0.0f, 10.0f, 0.0f, 0.0f)
+		[
+			SAssignNew(ButtonGrid, SUniformGridPanel)
+			.SlotPadding(FCoreStyle::Get().GetMargin("StandardDialog.SlotPadding"))
+			.MinDesiredSlotWidth(FCoreStyle::Get().GetFloat("StandardDialog.MinDesiredSlotWidth"))
+			.MinDesiredSlotHeight(FCoreStyle::Get().GetFloat("StandardDialog.MinDesiredSlotHeight"))
+			
+			+ SUniformGridPanel::Slot(0, 0)
 			[
-				SNew(SHorizontalBox)
+				// 'Prev' button
+				SNew(SButton)
+				.ButtonStyle(InArgs._ButtonStyle)
+				.TextStyle(InArgs._ButtonTextStyle)
+				.ForegroundColor(InArgs._ForegroundColor)
+				.HAlign(HAlign_Center)
+				.VAlign(VAlign_Center)
+				.ContentPadding(FCoreStyle::Get().GetMargin("StandardDialog.ContentPadding"))
+				.IsEnabled(this, &SWizard::HandlePrevButtonIsEnabled)
+				.OnClicked(this, &SWizard::HandlePrevButtonClicked)
+				.Visibility(this, &SWizard::HandlePrevButtonVisibility)
+				.ToolTipText(LOCTEXT("PrevButtonTooltip", "Go back to the previous step"))
+				[
+					SNew(SHorizontalBox)
 
-				+ SHorizontalBox::Slot()
+					+ SHorizontalBox::Slot()
+					.Padding(2.0f, 0.0f)
 					.AutoWidth()
-					.Padding(0.0f, 0.0f, 20.0f, 0.0f)
+					.VAlign(VAlign_Center)
 					[
-						SAssignNew(PageListBox, SVerticalBox)
-							.Visibility(InArgs._ShowPageList ? EVisibility::Visible : EVisibility::Collapsed)
+						SNew(SImage)
+						.Image(FCoreStyle::Get().GetBrush("Wizard.BackIcon"))
+						.ColorAndOpacity(FLinearColor(0.05f, 0.05f, 0.05f))
 					]
 
-				+ SHorizontalBox::Slot()
-					.FillWidth(1.0f)
+					+ SHorizontalBox::Slot()
+					.AutoWidth()
+					.VAlign(VAlign_Center)
 					[
-						// widget switcher
-						SAssignNew(WidgetSwitcher, SWidgetSwitcher)
+						SNew(STextBlock)
+						.Text(LOCTEXT("PrevButtonLabel", "Back"))
 					]
+				]
 			]
 
-		+ SVerticalBox::Slot()
-			.AutoHeight()
-			.Padding(0)
-			[
-				InArgs._PageFooter.Widget
-			]
-
-		+ SVerticalBox::Slot()
-			.AutoHeight()
-			.HAlign(HAlign_Right)
-			.Padding(0.0f, 10.0f, 0.0f, 0.0f)
+			+ SUniformGridPanel::Slot(1, 0)
 			[
-				SAssignNew(ButtonGrid, SUniformGridPanel)
-					.SlotPadding(FCoreStyle::Get().GetMargin("StandardDialog.SlotPadding"))
-					.MinDesiredSlotWidth(FCoreStyle::Get().GetFloat("StandardDialog.MinDesiredSlotWidth"))
-					.MinDesiredSlotHeight(FCoreStyle::Get().GetFloat("StandardDialog.MinDesiredSlotHeight"))
+				// 'Next' button
+				SNew(SButton)
+				.ButtonStyle(InArgs._ButtonStyle)
+				.TextStyle(InArgs._ButtonTextStyle)
+				.ForegroundColor(InArgs._ForegroundColor)
+				.HAlign(HAlign_Center)
+				.VAlign(VAlign_Center)
+				.ContentPadding(FCoreStyle::Get().GetMargin("StandardDialog.ContentPadding"))
+				.IsEnabled(this, &SWizard::HandleNextButtonIsEnabled)
+				.OnClicked(this, &SWizard::HandleNextButtonClicked)
+				.Visibility(this, &SWizard::HandleNextButtonVisibility)
+				.ToolTipText(LOCTEXT("NextButtonTooltip", "Go to the next step"))
+				[
+					SNew(SHorizontalBox)
 
-				+ SUniformGridPanel::Slot(0, 0)
+					+SHorizontalBox::Slot()
+					.AutoWidth()
+					.VAlign(VAlign_Center)
 					[
-						// 'Prev' button
-						SNew(SButton)
-							.HAlign(HAlign_Center)
-							.VAlign(VAlign_Center)
-							.ContentPadding(FCoreStyle::Get().GetMargin("StandardDialog.ContentPadding"))
-							.IsEnabled(this, &SWizard::HandlePrevButtonIsEnabled)
-							.OnClicked(this, &SWizard::HandlePrevButtonClicked)
-							.Visibility(this, &SWizard::HandlePrevButtonVisibility)
-							.ToolTipText(LOCTEXT("PrevButtonTooltip", "Go back to the previous step"))
-							[
-								SNew(SHorizontalBox)
-
-								+ SHorizontalBox::Slot()
-									.Padding(2.0f, 0.0f)
-									.AutoWidth()
-									.VAlign(VAlign_Center)
-									[
-										SNew(SImage)
-											.Image(FCoreStyle::Get().GetBrush("Wizard.BackIcon"))
-											.ColorAndOpacity(FLinearColor(0.05f, 0.05f, 0.05f))
-									]
-
-								+ SHorizontalBox::Slot()
-									.AutoWidth()
-									.VAlign(VAlign_Center)
-									[
-										SNew(STextBlock)
-											.Text(LOCTEXT("PrevButtonLabel", "Back"))
-									]
-							]
+						SNew(STextBlock)
+						.Text(LOCTEXT("NextButtonLabel", "Next"))
 					]
 
-				+ SUniformGridPanel::Slot(1, 0)
-				[
-					// 'Next' button
-					SNew(SButton)
-					.HAlign(HAlign_Center)
+					+SHorizontalBox::Slot()
+					.Padding(2.0f, 0.0f)
+					.AutoWidth()
 					.VAlign(VAlign_Center)
-					.ContentPadding(FCoreStyle::Get().GetMargin("StandardDialog.ContentPadding"))
-					.IsEnabled(this, &SWizard::HandleNextButtonIsEnabled)
-					.OnClicked(this, &SWizard::HandleNextButtonClicked)
-					.Visibility(this, &SWizard::HandleNextButtonVisibility)
-					.ToolTipText(LOCTEXT("NextButtonTooltip", "Go to the next step"))
 					[
-						SNew(SHorizontalBox)
-
-						+SHorizontalBox::Slot()
-						.AutoWidth()
-						.VAlign(VAlign_Center)
-						[
-							SNew(STextBlock)
-							.Text(LOCTEXT("NextButtonLabel", "Next"))
-						]
-
-						+SHorizontalBox::Slot()
-						.Padding(2.0f, 0.0f)
-						.AutoWidth()
-						.VAlign(VAlign_Center)
-						[
-							SNew(SImage)
-							.Image(FCoreStyle::Get().GetBrush("Wizard.NextIcon"))
-							.ColorAndOpacity(FLinearColor(0.05f, 0.05f, 0.05f))
-						]
+						SNew(SImage)
+						.Image(FCoreStyle::Get().GetBrush("Wizard.NextIcon"))
+						.ColorAndOpacity(FLinearColor(0.05f, 0.05f, 0.05f))
 					]
 				]
+			]
 
-				+ SUniformGridPanel::Slot(2,0)
-					[
-						// 'Finish' button
-						SNew(SButton)
-							.HAlign(HAlign_Center)
-							.VAlign(VAlign_Center)
-							.ContentPadding(FCoreStyle::Get().GetMargin("StandardDialog.ContentPadding"))
-							.IsEnabled(InArgs._CanFinish)
-							.OnClicked(this, &SWizard::HandleFinishButtonClicked)
-							.ToolTipText(InArgs._FinishButtonToolTip)
-							.Text(InArgs._FinishButtonText)
-					]
+			+ SUniformGridPanel::Slot(2,0)
+			[
+				// 'Finish' button
+				SNew(SButton)
+				.ButtonStyle(InArgs._FinishButtonStyle)
+				.TextStyle(InArgs._ButtonTextStyle)
+				.ForegroundColor(InArgs._ForegroundColor)
+				.HAlign(HAlign_Center)
+				.VAlign(VAlign_Center)
+				.ContentPadding(FCoreStyle::Get().GetMargin("StandardDialog.ContentPadding"))
+				.IsEnabled(InArgs._CanFinish)
+				.OnClicked(this, &SWizard::HandleFinishButtonClicked)
+				.ToolTipText(InArgs._FinishButtonToolTip)
+				.Text(InArgs._FinishButtonText)
 			]
+		]
 	];
 
 	if (InArgs._ShowCancelButton)
@@ -163,12 +171,15 @@ void SWizard::Construct( const FArguments& InArgs )
 		[
 			// cancel button
 			SNew(SButton)
-				.HAlign(HAlign_Center)
-				.VAlign(VAlign_Center)
-				.ContentPadding(FCoreStyle::Get().GetMargin("StandardDialog.ContentPadding"))
-				.OnClicked(this, &SWizard::HandleCancelButtonClicked)
-				.ToolTipText(LOCTEXT("CancelButtonTooltip", "Cancel this wizard"))
-				.Text(LOCTEXT("CancelButtonLabel", "Cancel"))
+			.ButtonStyle(InArgs._CancelButtonStyle)
+			.TextStyle(InArgs._ButtonTextStyle)
+			.ForegroundColor(InArgs._ForegroundColor)
+			.HAlign(HAlign_Center)
+			.VAlign(VAlign_Center)
+			.ContentPadding(FCoreStyle::Get().GetMargin("StandardDialog.ContentPadding"))
+			.OnClicked(this, &SWizard::HandleCancelButtonClicked)
+			.ToolTipText(LOCTEXT("CancelButtonTooltip", "Cancel this wizard"))
+			.Text(LOCTEXT("CancelButtonLabel", "Cancel"))
 		];
 	}
 
@@ -182,31 +193,30 @@ void SWizard::Construct( const FArguments& InArgs )
 		if (InArgs._ShowPageList)
 		{
 			PageListBox->AddSlot()
-				.AutoHeight()
+			.AutoHeight()
+			[
+				SNew(SCheckBox)
+				.IsChecked(this, &SWizard::HandlePageButtonIsChecked, SlotIndex)
+				.IsEnabled(this, &SWizard::HandlePageButtonIsEnabled, SlotIndex)
+				.OnCheckStateChanged(this, &SWizard::HandlePageButtonCheckStateChanged, SlotIndex)
+				.Padding(FMargin(8.0f, 4.0f, 24.0f, 4.0f))
+				.Style(&FCoreStyle::Get().GetWidgetStyle< FCheckBoxStyle >("ToggleButtonCheckbox"))
 				[
-					SNew(SCheckBox)
-						.IsChecked(this, &SWizard::HandlePageButtonIsChecked, SlotIndex)
-						.IsEnabled(this, &SWizard::HandlePageButtonIsEnabled, SlotIndex)
-						.OnCheckStateChanged(this, &SWizard::HandlePageButtonCheckStateChanged, SlotIndex)
-						.Padding(FMargin(8.0f, 4.0f, 24.0f, 4.0f))
-						.Style(&FCoreStyle::Get().GetWidgetStyle< FCheckBoxStyle >("ToggleButtonCheckbox"))
-						[
-							Page->GetButtonContent()
-						]
-				];
+					Page->GetButtonContent()
+				]
+			];
 		}
 
 		WidgetSwitcher->AddSlot()
-			[
-				Page->GetPageContent()
-			];
+		[
+			Page->GetPageContent()
+		];
 	}
 
 	WidgetSwitcher->SetActiveWidgetIndex(INDEX_NONE);
 	ShowPage(InArgs._InitialPageIndex.Get());
 }
 
-
 void SWizard::ShowPage( int32 PageIndex )
 {
 	int32 ActivePageIndex = WidgetSwitcher->GetActiveWidgetIndex();
diff --git a/Engine/Source/Runtime/AppFramework/Public/Widgets/Workflow/SWizard.h b/Engine/Source/Runtime/AppFramework/Public/Widgets/Workflow/SWizard.h
--- a/Engine/Source/Runtime/AppFramework/Public/Widgets/Workflow/SWizard.h
+++ b/Engine/Source/Runtime/AppFramework/Public/Widgets/Workflow/SWizard.h
@@ -138,7 +138,12 @@ public:
 public:
 
 	SLATE_BEGIN_ARGS(SWizard)
-		: _CanFinish(true)
+		: _ButtonStyle(&FCoreStyle::Get().GetWidgetStyle<FButtonStyle>("Button"))
+		, _FinishButtonStyle(&FCoreStyle::Get().GetWidgetStyle<FButtonStyle>("Button"))
+		, _CancelButtonStyle(&FCoreStyle::Get().GetWidgetStyle<FButtonStyle>("Button"))
+		, _ButtonTextStyle(&FCoreStyle::Get().GetWidgetStyle< FTextBlockStyle >("NormalText"))
+		, _ForegroundColor(FCoreStyle::Get().GetSlateColor("InvertedForeground"))
+		, _CanFinish(true)
 		, _FinishButtonText(NSLOCTEXT("SWizard", "DefaultFinishButtonText", "Finish"))
 		, _FinishButtonToolTip(NSLOCTEXT("SWizard", "DefaultFinishButtonTooltip", "Finish the wizard"))
 		, _InitialPageIndex(0)
@@ -150,6 +155,21 @@ public:
 
 		SLATE_SUPPORTS_SLOT_WITH_ARGS(FWizardPage)
 
+		/** The button style used by next and back. */
+		SLATE_STYLE_ARGUMENT(FButtonStyle, ButtonStyle)
+
+		/** The button style used by the cancel button. */
+		SLATE_STYLE_ARGUMENT(FButtonStyle, CancelButtonStyle)
+
+		/** The button style used by the finish button. */
+		SLATE_STYLE_ARGUMENT(FButtonStyle, FinishButtonStyle)
+
+		/** The text style used by the buttons. */
+		SLATE_STYLE_ARGUMENT(FTextBlockStyle, ButtonTextStyle)
+
+		/** The text style used by the buttons. */
+		SLATE_ARGUMENT(FSlateColor, ForegroundColor)
+
 		/** Holds a flag indicating whether the 'Finish' button is enabled. */
 		SLATE_ATTRIBUTE(bool, CanFinish)
 
