Slate: Improve SWizard styling compatibility (now obeys text style on Next/Back buttons, and uses the success style on Next)

diff --git a/Engine/Source/Runtime/AppFramework/Private/Widgets/Workflow/SWizard.cpp b/Engine/Source/Runtime/AppFramework/Private/Widgets/Workflow/SWizard.cpp
--- a/Engine/Source/Runtime/AppFramework/Private/Widgets/Workflow/SWizard.cpp
+++ b/Engine/Source/Runtime/AppFramework/Private/Widgets/Workflow/SWizard.cpp
@@ -105,6 +105,7 @@ void SWizard::Construct( const FArguments& InArgs )
 					.VAlign(VAlign_Center)
 					[
 						SNew(STextBlock)
+						.TextStyle(InArgs._ButtonTextStyle)
 						.Text(LOCTEXT("PrevButtonLabel", "Back"))
 					]
 				]
@@ -114,7 +115,7 @@ void SWizard::Construct( const FArguments& InArgs )
 			[
 				// 'Next' button
 				SNew(SButton)
-				.ButtonStyle(InArgs._ButtonStyle)
+				.ButtonStyle(InArgs._FinishButtonStyle)
 				.TextStyle(InArgs._ButtonTextStyle)
 				.ForegroundColor(InArgs._ForegroundColor)
 				.HAlign(HAlign_Center)
@@ -132,6 +133,7 @@ void SWizard::Construct( const FArguments& InArgs )
 					.VAlign(VAlign_Center)
 					[
 						SNew(STextBlock)
+						.TextStyle(InArgs._ButtonTextStyle)
 						.Text(LOCTEXT("NextButtonLabel", "Next"))
 					]
 
