Added test case for multi-line editable hint text

diff --git a/Engine/Source/Runtime/AppFramework/Private/Widgets/Testing/STestSuite.cpp b/Engine/Source/Runtime/AppFramework/Private/Widgets/Testing/STestSuite.cpp
--- a/Engine/Source/Runtime/AppFramework/Private/Widgets/Testing/STestSuite.cpp
+++ b/Engine/Source/Runtime/AppFramework/Private/Widgets/Testing/STestSuite.cpp
@@ -1181,6 +1181,7 @@ class SMultiLineEditingTest : public SCompoundWidget
 						[
 							SNew( SMultiLineEditableTextBox )
 							.Font( FSlateFontInfo( FPaths::EngineContentDir() / TEXT( "Slate/Fonts/Roboto-Regular.ttf" ), 12 ) )
+							.HintText( LOCTEXT("MultiLineHintText", "This is hint line 1\nThis is hint line 2") )
 							//.WrapTextAt(300.0f)
 							//.Justification( ETextJustify::Right )
 						]
