Changed some text input boxes to be numeric input boxes, and fixed a crash when testing an invalid item index

diff --git a/Engine/Source/Runtime/AppFramework/Private/Framework/Testing/STableViewTesting.cpp b/Engine/Source/Runtime/AppFramework/Private/Framework/Testing/STableViewTesting.cpp
--- a/Engine/Source/Runtime/AppFramework/Private/Framework/Testing/STableViewTesting.cpp
+++ b/Engine/Source/Runtime/AppFramework/Private/Framework/Testing/STableViewTesting.cpp
@@ -2,7 +2,7 @@
 
 #include "AppFrameworkPrivatePCH.h"
 #include "STableViewTesting.h"
-
+#include "SNumericEntryBox.h"
 
 #define LOCTEXT_NAMESPACE "STableViewTesting"
 
@@ -699,9 +699,9 @@ public:
 				+SHorizontalBox::Slot().AutoWidth()
 				[
 					// We set how many data items we want in our test case.
-					SNew(SEditableTextBox)
-					.Text( this, &STableViewTesting::GetNumTotalItems )
-					.OnTextChanged( this, &STableViewTesting::NumItems_OnTextChanged )
+					SNew(SNumericEntryBox<int32>)
+					.Value( this, &STableViewTesting::GetNumTotalItems )
+					.OnValueChanged( this, &STableViewTesting::NumItems_OnValueChanged )
 				]
 				+SHorizontalBox::Slot().AutoWidth()
 				[
@@ -713,9 +713,9 @@ public:
 				+SHorizontalBox::Slot().AutoWidth()
 				[
 					// We set how many data items we want in our test case.
-					SNew(SEditableTextBox)
-					.Text( this, &STableViewTesting::GetScrollToIndex )
-					.OnTextChanged( this, &STableViewTesting::ScrollToIndex_OnTextChanged )
+					SNew(SNumericEntryBox<int32>)
+					.Value( this, &STableViewTesting::GetScrollToIndex )
+					.OnValueChanged( this, &STableViewTesting::ScrollToIndex_OnValueChanged )
 				]
 				+SHorizontalBox::Slot().AutoWidth()
 				[
@@ -1052,37 +1052,41 @@ protected:
 	}
 	
 	/** The textbox representing the number of data items changed */
-	void NumItems_OnTextChanged( const FText& InNewText )
+	void NumItems_OnValueChanged( int32 InNewValue )
 	{
-		TotalItems = FCString::Atoi( *InNewText.ToString() );
+		TotalItems = InNewValue;
 	}
 
-	FText GetScrollToIndex() const
+	TOptional<int32> GetScrollToIndex() const
 	{
-		return FText::AsNumber(ScrollToIndex);
+		return ScrollToIndex;
 	}
 	
-	void ScrollToIndex_OnTextChanged( const FText& NewScrollToIndex )
+	void ScrollToIndex_OnValueChanged( int32 NewScrollToIndex )
 	{
-		ScrollToIndex = FCString::Atoi( *NewScrollToIndex.ToString() );
+		ScrollToIndex = NewScrollToIndex;
 	}
 	
 	FReply ScrollToIndex_OnClicked()
 	{
-		if (TreeBeingTested.IsValid())
+		if (Items.IsValidIndex(ScrollToIndex))
 		{
-			TreeBeingTested->RequestScrollIntoView( Items[ ScrollToIndex ] );
-			TreeBeingTested->SetSelection( Items[ ScrollToIndex ] );
-		}
+			if (TreeBeingTested.IsValid())
+			{
+				TreeBeingTested->RequestScrollIntoView( Items[ ScrollToIndex ] );
+				TreeBeingTested->SetSelection( Items[ ScrollToIndex ] );
+			}
 
-		if ( TileViewBeingTested.IsValid() )
-		{
-			TileViewBeingTested->RequestScrollIntoView( Items[ ScrollToIndex ] );
-			TileViewBeingTested->SetSelection( Items[ ScrollToIndex ] );
+			if ( TileViewBeingTested.IsValid() )
+			{
+				TileViewBeingTested->RequestScrollIntoView( Items[ ScrollToIndex ] );
+				TileViewBeingTested->SetSelection( Items[ ScrollToIndex ] );
+			}
+
+			ListBeingTested->RequestScrollIntoView( Items[ ScrollToIndex ] );
+			ListBeingTested->SetSelection( Items[ ScrollToIndex ] );
 		}
 
-		ListBeingTested->RequestScrollIntoView( Items[ ScrollToIndex ] );
-		ListBeingTested->SetSelection( Items[ ScrollToIndex ] );
 		return FReply::Handled();
 	}
 	
@@ -1119,10 +1123,10 @@ protected:
 		return FReply::Handled();
 	}
 	
-	/** @return How many data items we want to be using as a string.*/
-	FText GetNumTotalItems() const
+	/** @return How many data items we want to be using.*/
+	TOptional<int32> GetNumTotalItems() const
 	{
-		return FText::AsNumber(TotalItems);
+		return TotalItems;
 	}
 	
 	/** @return Given a data item return a new widget to represent it in the ListView */
