Fixed gamma correction for the EyeDropper (UE-6942)

diff --git a/Engine/Source/Runtime/AppFramework/Private/Widgets/Colors/SEyeDropperButton.h b/Engine/Source/Runtime/AppFramework/Private/Widgets/Colors/SEyeDropperButton.h
--- a/Engine/Source/Runtime/AppFramework/Private/Widgets/Colors/SEyeDropperButton.h
+++ b/Engine/Source/Runtime/AppFramework/Private/Widgets/Colors/SEyeDropperButton.h
@@ -83,7 +83,7 @@ public:
 			if (CurrentCursorPosition != LastCursorPosition)
 			{
 				// In dropper mode and outside the button - sample the pixel color and push it to the client
-				FLinearColor ScreenColor = FPlatformMisc::GetScreenPixelColor(CurrentCursorPosition, false);
+				FLinearColor ScreenColor = FPlatformMisc::GetScreenPixelColor(CurrentCursorPosition, DisplayGamma.Get(1.0f));
 				OnValueChanged.ExecuteIfBound(ScreenColor);
 			}
 
