Editor - The color picker no longer assumes a straight inverse gamma conversion from colors sampled on the screen using the eye dropper.  This change makes it so sRGB is the assumed space all colors are sampled in, making sampling colors from Photoshop bring in the identical version of the color converted into linear space, which when transformed back into sRGB space will be the same color you saw in photoshop.  FColor and FLinearColors can now default to sRGB space using a UPROPERTY meta tag of sRGB="true" or false, you can force a specific default checking of the sRGB checkbox.  Several places in UMG now force this default since the space you really care about sampling and previewing is sRGB space since that's the space UMG will be rendered in (on PC), additional work will be needed on Mac where the final rendering space is gamma 2.2.

diff --git a/Engine/Source/Editor/DetailCustomizations/Private/MathStructCustomizations.cpp b/Engine/Source/Editor/DetailCustomizations/Private/MathStructCustomizations.cpp
--- a/Engine/Source/Editor/DetailCustomizations/Private/MathStructCustomizations.cpp
+++ b/Engine/Source/Editor/DetailCustomizations/Private/MathStructCustomizations.cpp
@@ -422,6 +422,11 @@ void FColorStructCustomization::CustomizeHeader( TSharedRef<class IPropertyHandl
 
 	bIsLinearColor = CastChecked<UStructProperty>( StructPropertyHandle->GetProperty() )->Struct->GetFName() == NAME_LinearColor;
 	bIgnoreAlpha = StructPropertyHandle->GetProperty()->HasMetaData(TEXT("HideAlphaChannel"));
+	
+	if ( StructPropertyHandle->GetProperty()->HasMetaData(TEXT("sRGB")) )
+	{
+		sRGBOverride = StructPropertyHandle->GetProperty()->GetBoolMetaData(TEXT("sRGB"));
+	}
 
 	auto PropertyUtils = StructCustomizationUtils.GetPropertyUtilities();
 	bDontUpdateWhileEditing = PropertyUtils.IsValid() ? PropertyUtils->DontUpdateValueWhileEditing() : false;
@@ -572,6 +577,7 @@ void FColorStructCustomization::CreateColorPicker( bool bUseAlpha, bool bOnlyRef
 	PickerArgs.bUseAlpha = !bIgnoreAlpha;
 	PickerArgs.bOnlyRefreshOnMouseUp = false;
 	PickerArgs.bOnlyRefreshOnOk = bRefreshOnlyOnOk;
+	PickerArgs.sRGBOverride = sRGBOverride;
 	PickerArgs.DisplayGamma = TAttribute<float>::Create( TAttribute<float>::FGetter::CreateUObject(GEngine, &UEngine::GetDisplayGamma) );
 	PickerArgs.OnColorCommitted = FOnLinearColorValueChanged::CreateSP( this, &FColorStructCustomization::OnSetColorFromColorPicker );
 	PickerArgs.OnColorPickerCancelled = FOnColorPickerCancelled::CreateSP( this, &FColorStructCustomization::OnColorPickerCancelled );
@@ -754,7 +760,7 @@ FReply FColorStructCustomization::OnMouseButtonDownColorBlock(const FGeometry& M
 		return FReply::Unhandled();
 	}
 	
-	CreateColorPicker( /*bUseAlpha*/ true, bDontUpdateWhileEditing);
+	CreateColorPicker( /*bUseAlpha*/ true, bDontUpdateWhileEditing );
 
 	//bIsInlineColorPickerVisible = !bIsInlineColorPickerVisible;
 
@@ -768,7 +774,7 @@ EVisibility FColorStructCustomization::GetInlineColorPickerVisibility() const
 
 FReply FColorStructCustomization::OnOpenFullColorPickerClicked()
 {
-	CreateColorPicker( /*bUseAlpha*/ true, bDontUpdateWhileEditing);
+	CreateColorPicker( /*bUseAlpha*/ true, bDontUpdateWhileEditing );
 
 	bIsInlineColorPickerVisible = false;
 
diff --git a/Engine/Source/Editor/DetailCustomizations/Private/MathStructCustomizations.h b/Engine/Source/Editor/DetailCustomizations/Private/MathStructCustomizations.h
--- a/Engine/Source/Editor/DetailCustomizations/Private/MathStructCustomizations.h
+++ b/Engine/Source/Editor/DetailCustomizations/Private/MathStructCustomizations.h
@@ -227,7 +227,7 @@ protected:
 	/** Called when the user clicks on the the button to get the full color picker */
 	FReply OnOpenFullColorPickerClicked();
 
-private:
+protected:
 	/** Saved per struct colors in case the user clicks cancel in the color picker */
 	TArray<FLinearColor> SavedPreColorPickerColors;
 	/** Color struct handle */
@@ -244,5 +244,7 @@ private:
 	TSharedPtr<SWidget> ColorPickerParentWidget;
 	/** The value won;t be updated while editing */
 	bool bDontUpdateWhileEditing;
+	/** Overrides the default state of the sRGB checkbox */
+	TOptional<bool> sRGBOverride;
 };
 
diff --git a/Engine/Source/Editor/DetailCustomizations/Private/SlateColorCustomization.cpp b/Engine/Source/Editor/DetailCustomizations/Private/SlateColorCustomization.cpp
--- a/Engine/Source/Editor/DetailCustomizations/Private/SlateColorCustomization.cpp
+++ b/Engine/Source/Editor/DetailCustomizations/Private/SlateColorCustomization.cpp
@@ -25,6 +25,9 @@ void FSlateColorCustomization::CustomizeHeader(TSharedRef<class IPropertyHandle>
 	SpecifiedColorHandle->SetOnPropertyValueChanged(FSimpleDelegate::CreateSP(SharedThis(this), &FSlateColorCustomization::OnValueChanged));
 
 	FColorStructCustomization::CustomizeHeader(SpecifiedColorHandle.ToSharedRef(), InHeaderRow, StructCustomizationUtils);
+
+	// Slate brushes always default to sRGB mode.
+	sRGBOverride = true;
 }
 
 void FSlateColorCustomization::MakeHeaderRow(TSharedRef<class IPropertyHandle>& InStructPropertyHandle, FDetailWidgetRow& Row)
diff --git a/Engine/Source/Runtime/AppFramework/Private/Widgets/Colors/SColorPicker.cpp b/Engine/Source/Runtime/AppFramework/Private/Widgets/Colors/SColorPicker.cpp
--- a/Engine/Source/Runtime/AppFramework/Private/Widgets/Colors/SColorPicker.cpp
+++ b/Engine/Source/Runtime/AppFramework/Private/Widgets/Colors/SColorPicker.cpp
@@ -63,6 +63,12 @@ void SColorPicker::Construct( const FArguments& InArgs )
 	DisplayGamma = InArgs._DisplayGamma;
 	bClosedViaOkOrCancel = false;
 
+	if ( InArgs._sRGBOverride.IsSet() )
+	{
+		OriginalSRGBOption = SColorThemesViewer::bSRGBEnabled;
+		SColorThemesViewer::bSRGBEnabled = InArgs._sRGBOverride.GetValue();
+	}
+
 	RegisterActiveTimer( 0.f, FWidgetActiveTimerDelegate::CreateSP( this, &SColorPicker::AnimatePostConstruct ) );
 
 	// We need a parent window to set the close callback
@@ -251,184 +257,184 @@ void SColorPicker::GenerateDefaultColorPickerContent( bool bAdvancedSectionExpan
 		SNew(SVerticalBox)
 
 		+ SVerticalBox::Slot()
-			.AutoHeight()
+		.AutoHeight()
+		[
+			SNew(SGridPanel)
+			.FillColumn(0, 1.0f)
+
+			+ SGridPanel::Slot(0, 0)
+			.Padding(0.0f, 1.0f, 20.0f, 1.0f)
+			[
+				SNew(SHorizontalBox)
+
+				+ SHorizontalBox::Slot()
+				.FillWidth(1.0f)
+				.Padding(0.0f, 1.0f)
+				[
+					SNew(SOverlay)
+
+					+ SOverlay::Slot()
+					[
+						// color theme bar
+						SAssignNew(CurrentThemeBar, SThemeColorBlocksBar)
+						.ColorTheme(this, &SColorPicker::HandleThemeBarColorTheme)
+						.EmptyText(LOCTEXT("EmptyBarHint", "Drag & drop colors here to save"))
+						.HideTrashCallback(this, &SColorPicker::HideSmallTrash)
+						.ShowTrashCallback(this, &SColorPicker::ShowSmallTrash)
+						.ToolTipText(LOCTEXT("CurrentThemeBarToolTip", "Current Color Theme"))
+						.UseAlpha(SharedThis(this), &SColorPicker::HandleThemeBarUseAlpha)
+						.UseSRGB(SharedThis(this), &SColorPicker::HandleColorPickerUseSRGB)
+						.OnSelectColor(this, &SColorPicker::HandleThemeBarColorSelected)
+					]
+
+					// hack: need to fix SThemeColorBlocksBar::EmptyText to render properly
+					+ SOverlay::Slot()
+					.HAlign(HAlign_Center)
+					.VAlign(VAlign_Center)
+					[
+						SNew(STextBlock)
+						.Text(LOCTEXT("EmptyBarHint", "Drag & drop colors here to save"))
+						.Visibility(this, &SColorPicker::HandleThemeBarHintVisibility)
+					]
+				]
+
+				+ SHorizontalBox::Slot()
+				.AutoWidth()
+				[
+					// color theme selector
+					SAssignNew(ColorThemeButtonOrSmallTrash, SBorder)
+					.BorderImage(FStyleDefaults::GetNoBrush())
+					.Padding(0.0f)
+				]
+			]
+
+			+ SGridPanel::Slot(1, 0)
+			.HAlign(HAlign_Right)
+			.VAlign(VAlign_Center)
 			[
-				SNew(SGridPanel)
-					.FillColumn(0, 1.0f)
+				// sRGB check box
+				SNew(SCheckBox)
+				.ToolTipText(LOCTEXT("SRGBCheckboxToolTip", "Toggle gamma corrected sRGB previewing"))
+				.IsChecked(this, &SColorPicker::HandleSRGBCheckBoxIsChecked)
+				.OnCheckStateChanged(this, &SColorPicker::HandleSRGBCheckBoxCheckStateChanged)
+				[
+					SNew(STextBlock)
+					.Text(LOCTEXT("SRGBCheckboxLabel", "sRGB"))
+				]
+			]
+
+			+ SGridPanel::Slot(0, 1)
+			.Padding(0.0f, 8.0f, 20.0f, 0.0f)
+			[
+				SNew(SBorder)
+				.BorderImage(FCoreStyle::Get().GetBrush("NoBorder"))
+				.Padding(0.0f)
+				.OnMouseButtonDown(this, &SColorPicker::HandleColorAreaMouseDown)
+				[
+					SNew(SOverlay)
 
-				+ SGridPanel::Slot(0, 0)
-					.Padding(0.0f, 1.0f, 20.0f, 1.0f)
+					// color wheel
+					+ SOverlay::Slot()
 					[
 						SNew(SHorizontalBox)
 
 						+ SHorizontalBox::Slot()
-							.FillWidth(1.0f)
-							.Padding(0.0f, 1.0f)
-							[
-								SNew(SOverlay)
-
-								+ SOverlay::Slot()
-									[
-										// color theme bar
-										SAssignNew(CurrentThemeBar, SThemeColorBlocksBar)
-											.ColorTheme(this, &SColorPicker::HandleThemeBarColorTheme)
-											.EmptyText(LOCTEXT("EmptyBarHint", "Drag & drop colors here to save"))
-											.HideTrashCallback(this, &SColorPicker::HideSmallTrash)
-											.ShowTrashCallback(this, &SColorPicker::ShowSmallTrash)
-											.ToolTipText(LOCTEXT("CurrentThemeBarToolTip", "Current Color Theme"))
-											.UseAlpha(SharedThis(this), &SColorPicker::HandleThemeBarUseAlpha)
-											.UseSRGB(SharedThis(this), &SColorPicker::HandleColorPickerUseSRGB)
-											.OnSelectColor(this, &SColorPicker::HandleThemeBarColorSelected)
-									]
-
-								// hack: need to fix SThemeColorBlocksBar::EmptyText to render properly
-								+ SOverlay::Slot()
-									.HAlign(HAlign_Center)
-									.VAlign(VAlign_Center)
-									[
-										SNew(STextBlock)
-											.Text(LOCTEXT("EmptyBarHint", "Drag & drop colors here to save"))
-											.Visibility(this, &SColorPicker::HandleThemeBarHintVisibility)
-									]
-							]
+						.FillWidth(1.0f)
+						.HAlign(HAlign_Center)
+						[
+							SNew(SColorWheel)
+							.SelectedColor(this, &SColorPicker::GetCurrentColor)
+							.Visibility(this, &SColorPicker::HandleColorPickerModeVisibility, EColorPickerModes::Wheel)
+							.OnValueChanged(this, &SColorPicker::HandleColorSpectrumValueChanged)
+							.OnMouseCaptureBegin(this, &SColorPicker::HandleInteractiveChangeBegin)
+							.OnMouseCaptureEnd(this, &SColorPicker::HandleInteractiveChangeEnd)
+						]
 
 						+ SHorizontalBox::Slot()
-							.AutoWidth()
-							[
-								// color theme selector
-								SAssignNew(ColorThemeButtonOrSmallTrash, SBorder)
-									.BorderImage(FStyleDefaults::GetNoBrush())
-									.Padding(0.0f)							
-							]
+						.AutoWidth()
+						.Padding(4.0f, 0.0f)
+						[
+							// saturation slider
+							MakeColorSlider(EColorPickerChannels::Saturation)
+						]
+
+						+ SHorizontalBox::Slot()
+						.AutoWidth()
+						[
+							// value slider
+							MakeColorSlider(EColorPickerChannels::Value)
+						]
 					]
 
-				+ SGridPanel::Slot(1, 0)
-					.HAlign(HAlign_Right)
-					.VAlign(VAlign_Center)
+					// color spectrum
+					+ SOverlay::Slot()
 					[
-						// sRGB check box
-						SNew(SCheckBox)
-							.ToolTipText(LOCTEXT("SRGBCheckboxToolTip", "Toggle gamma corrected sRGB previewing"))
-							.IsChecked(this, &SColorPicker::HandleSRGBCheckBoxIsChecked)
-							.OnCheckStateChanged(this, &SColorPicker::HandleSRGBCheckBoxCheckStateChanged)
-							[
-								SNew(STextBlock)
-									.Text(LOCTEXT("SRGBCheckboxLabel", "sRGB"))
-							]
+						SNew(SBox)
+						.HeightOverride(200.0f)
+						.WidthOverride(292.0f)
+						[
+							SNew(SColorSpectrum)
+							.SelectedColor(this, &SColorPicker::GetCurrentColor)
+							.Visibility(this, &SColorPicker::HandleColorPickerModeVisibility, EColorPickerModes::Spectrum)
+							.OnValueChanged(this, &SColorPicker::HandleColorSpectrumValueChanged)
+							.OnMouseCaptureBegin(this, &SColorPicker::HandleInteractiveChangeBegin)
+							.OnMouseCaptureEnd(this, &SColorPicker::HandleInteractiveChangeEnd)
+						]
 					]
+				]
+			]
+
+			+ SGridPanel::Slot(1, 1)
+			.Padding(0.0f, 8.0f, 0.0f, 0.0f)
+			[
+				SNew(SVerticalBox)
 
-				+ SGridPanel::Slot(0, 1)
-					.Padding(0.0f, 8.0f, 20.0f, 0.0f)
+				+ SVerticalBox::Slot()
+				.AutoHeight()
+				[
+					SNew(SBox)
+					.HeightOverride(100.0f)
+					.WidthOverride(70.0f)
 					[
-						SNew(SBorder)
-							.BorderImage(FCoreStyle::Get().GetBrush("NoBorder"))
-							.Padding(0.0f)
-							.OnMouseButtonDown(this, &SColorPicker::HandleColorAreaMouseDown)
-							[
-								SNew(SOverlay)
-
-								// color wheel
-								+ SOverlay::Slot()
-									[
-										SNew(SHorizontalBox)
-
-										+ SHorizontalBox::Slot()
-											.FillWidth(1.0f)
-											.HAlign(HAlign_Center)
-											[
-												SNew(SColorWheel)
-													.SelectedColor(this, &SColorPicker::GetCurrentColor)
-													.Visibility(this, &SColorPicker::HandleColorPickerModeVisibility, EColorPickerModes::Wheel)
-													.OnValueChanged(this, &SColorPicker::HandleColorSpectrumValueChanged)
-													.OnMouseCaptureBegin(this, &SColorPicker::HandleInteractiveChangeBegin)
-													.OnMouseCaptureEnd(this, &SColorPicker::HandleInteractiveChangeEnd)
-											]
-
-										+ SHorizontalBox::Slot()
-											.AutoWidth()
-											.Padding(4.0f, 0.0f)
-											[
-												// saturation slider
-												MakeColorSlider(EColorPickerChannels::Saturation)
-											]
-
-										+ SHorizontalBox::Slot()
-											.AutoWidth()
-											[
-												// value slider
-												MakeColorSlider(EColorPickerChannels::Value)
-											]
-									]
-
-								// color spectrum
-								+ SOverlay::Slot()
-									[
-										SNew(SBox)
-											.HeightOverride(200.0f)
-											.WidthOverride(292.0f)
-											[
-												SNew(SColorSpectrum)
-													.SelectedColor(this, &SColorPicker::GetCurrentColor)
-													.Visibility(this, &SColorPicker::HandleColorPickerModeVisibility, EColorPickerModes::Spectrum)
-													.OnValueChanged(this, &SColorPicker::HandleColorSpectrumValueChanged)
-													.OnMouseCaptureBegin(this, &SColorPicker::HandleInteractiveChangeBegin)
-													.OnMouseCaptureEnd(this, &SColorPicker::HandleInteractiveChangeEnd)
-											]
-									]
-							]
+						// color preview
+						MakeColorPreviewBox()
 					]
+				]
 
-				+ SGridPanel::Slot(1, 1)
-					.Padding(0.0f, 8.0f, 0.0f, 0.0f)
-					[
-						SNew(SVerticalBox)
+				+ SVerticalBox::Slot()
+				.AutoHeight()
+				.Padding(0.0f, 16.0f, 0.0f, 0.0f)
+				.VAlign(VAlign_Top)
+				[
+					SNew(SHorizontalBox)
 
-						+ SVerticalBox::Slot()
-							.AutoHeight()
-							[
-								SNew(SBox)
-									.HeightOverride(100.0f)
-									.WidthOverride(70.0f)
-									[
-										// color preview
-										MakeColorPreviewBox()
-									]
-							]
+					+ SHorizontalBox::Slot()
+					.HAlign(HAlign_Left)
+					[
+						// mode selector
+						SNew(SButton)
+						.OnClicked(this, &SColorPicker::HandleColorPickerModeButtonClicked)
+						.Content()
+						[
+							SNew(SImage)
+							.Image(FCoreStyle::Get().GetBrush("ColorPicker.Mode"))
+							.ToolTipText(LOCTEXT("ColorPickerModeEToolTip", "Toggle between color wheel and color spectrum."))
+						]									
+					]
 
-						+ SVerticalBox::Slot()
-							.AutoHeight()
-							.Padding(0.0f, 16.0f, 0.0f, 0.0f)
-							.VAlign(VAlign_Top)
-							[
-								SNew(SHorizontalBox)
-
-								+ SHorizontalBox::Slot()
-									.HAlign(HAlign_Left)
-									[
-										// mode selector
-										SNew(SButton)
-											.OnClicked(this, &SColorPicker::HandleColorPickerModeButtonClicked)
-											.Content()
-											[
-												SNew(SImage)
-													.Image(FCoreStyle::Get().GetBrush("ColorPicker.Mode"))
-													.ToolTipText(LOCTEXT("ColorPickerModeEToolTip", "Toggle between color wheel and color spectrum."))
-											]									
-									]
-
-								+ SHorizontalBox::Slot()
-									.HAlign(HAlign_Right)
-									[
-										// eye dropper
-										SNew(SEyeDropperButton)
-											.OnValueChanged(this, &SColorPicker::HandleRGBColorChanged)
-											.OnBegin(this, &SColorPicker::HandleInteractiveChangeBegin)
-											.OnComplete(this, &SColorPicker::HandleEyeDropperButtonComplete)
-											.DisplayGamma(DisplayGamma)								
-									]
-							]
+					+ SHorizontalBox::Slot()
+					.HAlign(HAlign_Right)
+					[
+						// eye dropper
+						SNew(SEyeDropperButton)
+						.OnValueChanged(this, &SColorPicker::HandleRGBColorChanged)
+						.OnBegin(this, &SColorPicker::HandleInteractiveChangeBegin)
+						.OnComplete(this, &SColorPicker::HandleEyeDropperButtonComplete)
+						.DisplayGamma(DisplayGamma)
 					]
+				]
 			]
+		]
 
 		// advanced settings
 		+ SVerticalBox::Slot()
@@ -436,104 +442,104 @@ void SColorPicker::GenerateDefaultColorPickerContent( bool bAdvancedSectionExpan
 			.Padding(0.0f, 4.0f, 0.0f, 0.0f)
 			[
 				SNew(SExpandableArea)
-					.AreaTitle(LOCTEXT("AdvancedAreaTitle", "Advanced"))
-					.BorderBackgroundColor(FLinearColor::Transparent)
-					.InitiallyCollapsed(!bAdvancedSectionExpanded)
-					.OnAreaExpansionChanged(this, &SColorPicker::HandleAdvancedAreaExpansionChanged)
-					.Padding(FMargin(0.0f, 1.0f, 0.0f, 8.0f))
-					.BodyContent()
+				.AreaTitle(LOCTEXT("AdvancedAreaTitle", "Advanced"))
+				.BorderBackgroundColor(FLinearColor::Transparent)
+				.InitiallyCollapsed(!bAdvancedSectionExpanded)
+				.OnAreaExpansionChanged(this, &SColorPicker::HandleAdvancedAreaExpansionChanged)
+				.Padding(FMargin(0.0f, 1.0f, 0.0f, 8.0f))
+				.BodyContent()
+				[
+					SNew(SHorizontalBox)
+
+					// RGBA inputs
+					+ SHorizontalBox::Slot()
+						.Padding(0.0f, 0.0f, 4.0f, 0.0f)
+						[
+							SNew(SVerticalBox)
+
+							// Red
+							+ SVerticalBox::Slot()
+								[
+									MakeColorSpinBox(EColorPickerChannels::Red)
+								]
+
+							// Green
+							+ SVerticalBox::Slot()
+								.Padding(0.0f, 8.0f, 0.0f, 0.0f)
+								[
+									MakeColorSpinBox(EColorPickerChannels::Green)
+								]
+
+							// Blue
+							+ SVerticalBox::Slot()
+								.Padding(0.0f, 8.0f, 0.0f, 0.0f)
+								[
+									MakeColorSpinBox(EColorPickerChannels::Blue)
+								]
+
+							// Alpha
+							+ SVerticalBox::Slot()
+								.Padding(0.0f, 8.0f, 0.0f, 0.0f)
+								[
+									MakeColorSpinBox(EColorPickerChannels::Alpha)
+								]
+						]
+
+					// HSV & Hex inputs
+					+ SHorizontalBox::Slot()
+					.Padding(4.0f, 0.0f, 0.0f, 0.0f)
 					[
-						SNew(SHorizontalBox)
+						SNew(SVerticalBox)
+							
+						// Hue
+						+ SVerticalBox::Slot()
+						[
+							MakeColorSpinBox(EColorPickerChannels::Hue)
+						]
 
-						// RGBA inputs
-						+ SHorizontalBox::Slot()
+						// Saturation
+						+ SVerticalBox::Slot()
+						.Padding(0.0f, 8.0f, 0.0f, 0.0f)
+						[
+							MakeColorSpinBox(EColorPickerChannels::Saturation)
+						]
+
+						// Value
+						+ SVerticalBox::Slot()
+						.Padding(0.0f, 8.0f, 0.0f, 0.0f)
+						[
+							MakeColorSpinBox(EColorPickerChannels::Value)
+						]
+
+						// Hex
+						+ SVerticalBox::Slot()
+						.HAlign(HAlign_Right)
+						.VAlign(VAlign_Top)
+						.Padding(0.0f, 8.0f, 0.0f, 0.0f)
+						[
+							SNew(SHorizontalBox)
+							.ToolTipText(LOCTEXT("HexSliderToolTip", "Hexadecimal Value"))
+
+							+ SHorizontalBox::Slot()
+							.AutoWidth()
 							.Padding(0.0f, 0.0f, 4.0f, 0.0f)
+							.VAlign(VAlign_Center)
 							[
-								SNew(SVerticalBox)
-
-								// Red
-								+ SVerticalBox::Slot()
-									[
-										MakeColorSpinBox(EColorPickerChannels::Red)
-									]
-
-								// Green
-								+ SVerticalBox::Slot()
-									.Padding(0.0f, 8.0f, 0.0f, 0.0f)
-									[
-										MakeColorSpinBox(EColorPickerChannels::Green)
-									]
-
-								// Blue
-								+ SVerticalBox::Slot()
-									.Padding(0.0f, 8.0f, 0.0f, 0.0f)
-									[
-										MakeColorSpinBox(EColorPickerChannels::Blue)
-									]
-
-								// Alpha
-								+ SVerticalBox::Slot()
-									.Padding(0.0f, 8.0f, 0.0f, 0.0f)
-									[
-										MakeColorSpinBox(EColorPickerChannels::Alpha)
-									]
+								SNew(STextBlock)
+								.Text(LOCTEXT("HexInputLabel", "Hex"))
 							]
 
-						// HSV & Hex inputs
-						+ SHorizontalBox::Slot()
-							.Padding(4.0f, 0.0f, 0.0f, 0.0f)
+							+ SHorizontalBox::Slot()
+							.AutoWidth()
 							[
-								SNew(SVerticalBox)
-							
-								// Hue
-								+ SVerticalBox::Slot()
-									[
-										MakeColorSpinBox(EColorPickerChannels::Hue)
-									]
-
-								// Saturation
-								+ SVerticalBox::Slot()
-									.Padding(0.0f, 8.0f, 0.0f, 0.0f)
-									[
-										MakeColorSpinBox(EColorPickerChannels::Saturation)
-									]
-
-								// Value
-								+ SVerticalBox::Slot()
-									.Padding(0.0f, 8.0f, 0.0f, 0.0f)
-									[
-										MakeColorSpinBox(EColorPickerChannels::Value)
-									]
-
-								// Hex
-								+ SVerticalBox::Slot()
-									.HAlign(HAlign_Right)
-									.VAlign(VAlign_Top)
-									.Padding(0.0f, 8.0f, 0.0f, 0.0f)
-									[
-										SNew(SHorizontalBox)
-											.ToolTipText(LOCTEXT("HexSliderToolTip", "Hexadecimal Value"))
-
-										+ SHorizontalBox::Slot()
-											.AutoWidth()
-											.Padding(0.0f, 0.0f, 4.0f, 0.0f)
-											.VAlign(VAlign_Center)
-											[
-												SNew(STextBlock)
-													.Text(LOCTEXT("HexInputLabel", "Hex"))
-											]
-
-										+ SHorizontalBox::Slot()
-											.AutoWidth()
-											[
-												SNew(SEditableTextBox)
-													.MinDesiredWidth(72.0f)
-													.Text(this, &SColorPicker::HandleHexBoxText)
-													.OnTextCommitted(this, &SColorPicker::HandleHexInputTextCommitted)
-											]
-									]
+								SNew(SEditableTextBox)
+								.MinDesiredWidth(72.0f)
+								.Text(this, &SColorPicker::HandleHexBoxText)
+								.OnTextCommitted(this, &SColorPicker::HandleHexInputTextCommitted)
 							]
+						]
 					]
+				]
 			]
 
 		// dialog buttons
@@ -1413,6 +1419,11 @@ void SColorPicker::HandleParentWindowClosed( const TSharedRef<SWindow>& Window )
 	{
 		DiscardColor();
 	}
+
+	if ( OriginalSRGBOption.IsSet() )
+	{
+		SColorThemesViewer::bSRGBEnabled = OriginalSRGBOption.GetValue();
+	}
 }
 
 
@@ -1567,18 +1578,19 @@ bool OpenColorPicker(const FColorPickerArgs& Args)
 		.OnInteractivePickEnd(Args.OnInteractivePickEnd)
 		.OnColorPickerWindowClosed(Args.OnColorPickerWindowClosed)
 		.ParentWindow(Window)
-		.DisplayGamma(Args.DisplayGamma);
+		.DisplayGamma(Args.DisplayGamma)
+		.sRGBOverride(Args.sRGBOverride);
 		
 	Window->SetContent(
 		SNew(SBox)
+		[
+			SNew(SBorder)
+			.BorderImage(FCoreStyle::Get().GetBrush("ToolPanel.GroupBorder"))
+			.Padding(FMargin(8.0f, 8.0f))
 			[
-				SNew(SBorder)
-					.BorderImage(FCoreStyle::Get().GetBrush("ToolPanel.GroupBorder"))
-					.Padding(FMargin(8.0f, 8.0f))
-					[
-						ColorPicker
-					]
+				ColorPicker
 			]
+		]
 	);
 
 	if (Args.bIsModal)
diff --git a/Engine/Source/Runtime/AppFramework/Private/Widgets/Colors/SEyeDropperButton.cpp b/Engine/Source/Runtime/AppFramework/Private/Widgets/Colors/SEyeDropperButton.cpp
new file mode 100644
--- /dev/null
+++ b/Engine/Source/Runtime/AppFramework/Private/Widgets/Colors/SEyeDropperButton.cpp
@@ -0,0 +1,223 @@
+// Copyright 1998-2015 Epic Games, Inc. All Rights Reserved.
+
+#include "AppFrameworkPrivatePCH.h"
+#include "SEyeDropperButton.h"
+
+#define LOCTEXT_NAMESPACE "EyeDroppperButton"
+
+void SEyeDropperButton::Construct(const FArguments& InArgs)
+{
+	OnValueChanged = InArgs._OnValueChanged;
+	OnBegin = InArgs._OnBegin;
+	OnComplete = InArgs._OnComplete;
+	DisplayGamma = InArgs._DisplayGamma;
+
+	// This is a button containing an dropper image and text that tells the user to hit Esc.
+	// Their visibility and colors are changed according to whether dropper mode is active or not.
+	SButton::Construct(
+		SButton::FArguments()
+		.ContentPadding(1.0f)
+		.OnClicked(this, &SEyeDropperButton::OnClicked)
+		[
+			SNew(SOverlay)
+
+			+ SOverlay::Slot()
+			.Padding(FMargin(1.0f, 0.0f))
+			[
+				SNew(SImage)
+				.Image(FCoreStyle::Get().GetBrush("ColorPicker.EyeDropper"))
+				.ToolTipText(LOCTEXT("EyeDropperButton_ToolTip", "Activates the eye-dropper for selecting a colored pixel from any window."))
+				.ColorAndOpacity(this, &SEyeDropperButton::GetDropperImageColor)
+			]
+
+			+ SOverlay::Slot()
+			[
+				SNew(STextBlock)
+				.Text(LOCTEXT("EscapeCue", "Esc"))
+				.ToolTipText(LOCTEXT("EyeDropperEscapeCue_ToolTip", "Hit Escape key to stop the eye dropper"))
+				.Visibility(this, &SEyeDropperButton::GetEscapeTextVisibility)
+			]
+		]
+	);
+}
+
+EActiveTimerReturnType SEyeDropperButton::ActiveTick(double InCurrentTime, float InDeltaTime)
+{
+	if ( bWasClickActivated )
+	{
+		if ( HasMouseCapture() && bWasLeft && !bWasReEntered )
+		{
+			FVector2D CurrentCursorPosition = FSlateApplication::Get().GetCursorPos();
+			if ( CurrentCursorPosition != LastCursorPosition )
+			{
+				// In dropper mode and outside the button - sample the pixel color and push it to the client
+				// Convert the display gamma into a ratio of gamma from the default gamma.
+				float Gamma = DisplayGamma.Get(2.2f) / 2.2f;
+				FLinearColor ScreenColor = FPlatformMisc::GetScreenPixelColor(CurrentCursorPosition, Gamma);
+				OnValueChanged.ExecuteIfBound(ScreenColor);
+			}
+
+			LastCursorPosition = CurrentCursorPosition;
+		}
+
+		return EActiveTimerReturnType::Continue;
+	}
+	
+	LastCursorPosition.Reset();
+	return EActiveTimerReturnType::Stop;
+}
+
+FReply SEyeDropperButton::OnMouseButtonDown(const FGeometry& MyGeometry, const FPointerEvent& MouseEvent)
+{
+	// Clicking ANY mouse button when the dropper isn't active resets the active dropper states ready to activate
+	if ( !HasMouseCapture() )
+	{
+		bWasClickActivated = false;
+		bWasLeft = false;
+		bWasReEntered = false;
+	}
+
+	return SButton::OnMouseButtonDown(MyGeometry, MouseEvent);
+}
+
+FReply SEyeDropperButton::OnMouseButtonUp(const FGeometry& MyGeometry, const FPointerEvent& MouseEvent)
+{
+	// If a mouse click is completing and the dropper is active ALWAYS deactivate
+	bool bDeactivating = false;
+	if ( bWasClickActivated )
+	{
+		bDeactivating = true;
+	}
+
+	// bWasClicked is reset here because if it is set during SButton::OnMouseButtonUp()
+	// then the button was 'clicked' according to the usual rules. We might want to capture the
+	// mouse when the button is clicked but can't do it in the Clicked callback.
+	bWasClicked = false;
+	FReply Reply = SButton::OnMouseButtonUp(MyGeometry, MouseEvent);
+
+	// Switching dropper mode off?
+	if ( bDeactivating )
+	{
+		// These state flags clear dropper mode
+		bWasClickActivated = false;
+		bWasLeft = false;
+		bWasReEntered = false;
+
+		Reply.ReleaseMouseCapture();
+
+		OnComplete.ExecuteIfBound();
+	}
+	else if ( MouseEvent.GetEffectingButton() == EKeys::LeftMouseButton && bWasClicked )
+	{
+		// A normal LMB mouse click on the button occurred
+		// Set the initial dropper mode state and capture the mouse
+		bWasClickActivated = true;
+		bWasLeft = false;
+		bWasReEntered = false;
+
+		OnBegin.ExecuteIfBound();
+
+		Reply.CaptureMouse(this->AsShared());
+
+		RegisterActiveTimer(0, FWidgetActiveTimerDelegate::CreateSP(this, &SEyeDropperButton::ActiveTick));
+	}
+	bWasClicked = false;
+
+	return Reply;
+}
+
+FReply SEyeDropperButton::OnMouseMove(const FGeometry& MyGeometry, const FPointerEvent& MouseEvent)
+{
+	// If the mouse is captured and bWasClickActivated is set then we are in dropper mode
+	if ( HasMouseCapture() && bWasClickActivated )
+	{
+		if ( IsMouseOver(MyGeometry, MouseEvent) )
+		{
+			if ( bWasLeft )
+			{
+				// Mouse is over the button having left it once
+				bWasReEntered = true;
+			}
+		}
+		else
+		{
+			// Mouse is outside the button
+			bWasLeft = true;
+			bWasReEntered = false;
+		}
+	}
+
+	return SButton::OnMouseMove(MyGeometry, MouseEvent);
+}
+
+FReply SEyeDropperButton::OnKeyDown(const FGeometry& MyGeometry, const FKeyEvent& InKeyEvent)
+{
+	// Escape key when in dropper mode
+	if ( InKeyEvent.GetKey() == EKeys::Escape &&
+		HasMouseCapture() &&
+		bWasClickActivated )
+	{
+		// Clear the dropper mode states
+		bWasClickActivated = false;
+		bWasLeft = false;
+		bWasReEntered = false;
+
+		// This is needed to switch the dropper cursor off immediately so the user can see the Esc key worked
+		FSlateApplication::Get().QueryCursor();
+
+		FReply ReleaseReply = FReply::Handled().ReleaseMouseCapture();
+
+		OnComplete.ExecuteIfBound();
+
+		return ReleaseReply;
+	}
+
+	return FReply::Unhandled();
+}
+
+FCursorReply SEyeDropperButton::OnCursorQuery(const FGeometry& MyGeometry, const FPointerEvent& CursorEvent) const
+{
+	// Cursor is changed to the dropper when dropper mode is active and the states are correct
+	if ( HasMouseCapture() &&
+		bWasClickActivated && bWasLeft && !bWasReEntered )
+	{
+		return FCursorReply::Cursor(EMouseCursor::EyeDropper);
+	}
+
+	return SButton::OnCursorQuery(MyGeometry, CursorEvent);
+}
+
+FReply SEyeDropperButton::OnClicked()
+{
+	// Log clicks so that OnMouseButtonUp() can post process clicks
+	bWasClicked = true;
+	return FReply::Handled();
+}
+
+EVisibility SEyeDropperButton::GetEscapeTextVisibility() const
+{
+	// Show the Esc key message in the button when dropper mode is active
+	if ( HasMouseCapture() && bWasClickActivated )
+	{
+		return EVisibility::Visible;
+	}
+	return EVisibility::Hidden;
+}
+
+FSlateColor SEyeDropperButton::GetDropperImageColor() const
+{
+	// Make the dropper image in the button pale when dropper mode is active
+	if ( HasMouseCapture() && bWasClickActivated )
+	{
+		return FLinearColor(0.3f, 0.3f, 0.3f, 1.0f);
+	}
+	return FSlateColor::UseForeground();
+}
+
+bool SEyeDropperButton::IsMouseOver(const FGeometry& MyGeometry, const FPointerEvent& MouseEvent) const
+{
+	return MyGeometry.IsUnderLocation(MouseEvent.GetScreenSpacePosition());
+}
+
+
+#undef LOCTEXT_NAMESPACE
diff --git a/Engine/Source/Runtime/AppFramework/Private/Widgets/Colors/SEyeDropperButton.h b/Engine/Source/Runtime/AppFramework/Private/Widgets/Colors/SEyeDropperButton.h
--- a/Engine/Source/Runtime/AppFramework/Private/Widgets/Colors/SEyeDropperButton.h
+++ b/Engine/Source/Runtime/AppFramework/Private/Widgets/Colors/SEyeDropperButton.h
@@ -2,10 +2,6 @@
 
 #pragma once
 
-
-#define LOCTEXT_NAMESPACE "EyeDroppperButton"
-
-
 /**
  * Class for placing a color picker eye-dropper button.
  * A self contained until that only needs client code to set the display gamma and listen
@@ -26,7 +22,7 @@ public:
 		{}
 
 		/** Invoked when a new value is selected by the dropper */
-		SLATE_EVENT(FOnLinearColorValueChanged, OnValueChanged)
+	SLATE_EVENT(FOnLinearColorValueChanged, OnValueChanged)
 
 		/** Invoked when the dropper goes from inactive to active */
 		SLATE_EVENT(FSimpleDelegate, OnBegin)
@@ -39,217 +35,27 @@ public:
 
 	SLATE_END_ARGS()
 
-	void Construct(const FArguments& InArgs )
-	{
-		OnValueChanged = InArgs._OnValueChanged;
-		OnBegin = InArgs._OnBegin;
-		OnComplete = InArgs._OnComplete;
-		DisplayGamma = InArgs._DisplayGamma;
-
-		// This is a button containing an dropper image and text that tells the user to hit Esc.
-		// Their visibility and colors are changed according to whether dropper mode is active or not.
-		SButton::Construct(
-			SButton::FArguments()
-				.ContentPadding(1.0f)
-				.OnClicked(this, &SEyeDropperButton::OnClicked)
-				[
-					SNew(SOverlay)
-
-					+ SOverlay::Slot()
-						.Padding(FMargin(1.0f, 0.0f))
-						[
-							SNew(SImage)
-								.Image(FCoreStyle::Get().GetBrush("ColorPicker.EyeDropper"))
-								.ToolTipText(LOCTEXT("EyeDropperButton_ToolTip", "Activates the eye-dropper for selecting a colored pixel from any window."))
-								.ColorAndOpacity(this, &SEyeDropperButton::GetDropperImageColor)				
-						]
-
-					+ SOverlay::Slot()
-						[
-							SNew(STextBlock)
-								.Text(LOCTEXT("EscapeCue", "Esc"))
-								.ToolTipText(LOCTEXT("EyeDropperEscapeCue_ToolTip", "Hit Escape key to stop the eye dropper"))
-								.Visibility(this, &SEyeDropperButton::GetEscapeTextVisibility)
-						]
-				]
-		);
-	}
-
-	virtual void Tick(const FGeometry& AllottedGeometry, const double InCurrentTime, const float InDeltaTime) override
-	{
-		if (bWasClickActivated && HasMouseCapture() && bWasLeft && !bWasReEntered)
-		{
-			FVector2D CurrentCursorPosition = FSlateApplication::Get().GetCursorPos();
-			if (CurrentCursorPosition != LastCursorPosition)
-			{
-				// In dropper mode and outside the button - sample the pixel color and push it to the client
-				FLinearColor ScreenColor = FPlatformMisc::GetScreenPixelColor(CurrentCursorPosition, DisplayGamma.Get(1.0f));
-				OnValueChanged.ExecuteIfBound(ScreenColor);
-			}
-
-			LastCursorPosition = CurrentCursorPosition;
-		}
-		else
-		{
-			LastCursorPosition.Reset();
-		}
-
-		SButton::Tick(AllottedGeometry, InCurrentTime, InDeltaTime);
-	}
-
-
-	virtual FReply OnMouseButtonDown( const FGeometry& MyGeometry, const FPointerEvent& MouseEvent ) override
-	{
-		// Clicking ANY mouse button when the dropper isn't active resets the active dropper states ready to activate
-		if (!HasMouseCapture())
-		{
-			bWasClickActivated = false;
-			bWasLeft = false;
-			bWasReEntered = false;
-		}
-
-		return SButton::OnMouseButtonDown(MyGeometry, MouseEvent);
-	}
-
-	virtual FReply OnMouseButtonUp( const FGeometry& MyGeometry, const FPointerEvent& MouseEvent ) override
-	{
-		// If a mouse click is completing and the dropper is active ALWAYS deactivate
-		bool bDeactivating = false;
-		if (bWasClickActivated)
-		{
-			bDeactivating = true;
-		}
-
-		// bWasClicked is reset here because if it is set during SButton::OnMouseButtonUp()
-		// then the button was 'clicked' according to the usual rules. We might want to capture the
-		// mouse when the button is clicked but can't do it in the Clicked callback.
-		bWasClicked = false;
-		FReply Reply = SButton::OnMouseButtonUp(MyGeometry, MouseEvent);
-
-		// Switching dropper mode off?
-		if (bDeactivating)
-		{
-			// These state flags clear dropper mode
-			bWasClickActivated = false;
-			bWasLeft = false;
-			bWasReEntered = false;
+	void Construct(const FArguments& InArgs);
 
-			Reply.ReleaseMouseCapture();
+	virtual FReply OnMouseButtonDown(const FGeometry& MyGeometry, const FPointerEvent& MouseEvent) override;
+	virtual FReply OnMouseButtonUp(const FGeometry& MyGeometry, const FPointerEvent& MouseEvent) override;
+	virtual FReply OnMouseMove(const FGeometry& MyGeometry, const FPointerEvent& MouseEvent) override;
+	virtual FReply OnKeyDown(const FGeometry& MyGeometry, const FKeyEvent& InKeyEvent) override;
+	virtual FCursorReply OnCursorQuery(const FGeometry& MyGeometry, const FPointerEvent& CursorEvent) const override;
 
-			OnComplete.ExecuteIfBound();
-		}
-		else if (MouseEvent.GetEffectingButton() == EKeys::LeftMouseButton && bWasClicked)
-		{
-			// A normal LMB mouse click on the button occurred
-			// Set the initial dropper mode state and capture the mouse
-			bWasClickActivated = true;
-			bWasLeft = false;
-			bWasReEntered = false;
-
-			OnBegin.ExecuteIfBound();
-
-			Reply.CaptureMouse(this->AsShared());
-		}
-		bWasClicked = false;
-
-		return Reply;
-	}
-
-	virtual FReply OnMouseMove( const FGeometry& MyGeometry, const FPointerEvent& MouseEvent ) override
-	{
-		// If the mouse is captured and bWasClickActivated is set then we are in dropper mode
-		if (HasMouseCapture() && bWasClickActivated)
-		{
-			if (IsMouseOver(MyGeometry, MouseEvent))
-			{
-				if (bWasLeft)
-				{
-					// Mouse is over the button having left it once
-					bWasReEntered = true;
-				}
-			}
-			else
-			{
-				// Mouse is outside the button
-				bWasLeft = true;
-				bWasReEntered = false;
-			}
-		}
-
-		return SButton::OnMouseMove( MyGeometry, MouseEvent );
-	}
-
-	FReply OnKeyDown( const FGeometry& MyGeometry, const FKeyEvent& InKeyEvent ) override
-	{
-		// Escape key when in dropper mode
-		if (InKeyEvent.GetKey() == EKeys::Escape &&
-			HasMouseCapture() &&
-			bWasClickActivated)
-		{
-			// Clear the dropper mode states
-			bWasClickActivated = false;
-			bWasLeft = false;
-			bWasReEntered = false;
-
-			// This is needed to switch the dropper cursor off immediately so the user can see the Esc key worked
-			FSlateApplication::Get().QueryCursor();
-
-			FReply ReleaseReply = FReply::Handled().ReleaseMouseCapture();
-
-			OnComplete.ExecuteIfBound();
-
-			return ReleaseReply;
-		}
+private:
 
-		return FReply::Unhandled();
-	}
+	FReply OnClicked();
+	EActiveTimerReturnType ActiveTick(double InCurrentTime, float InDeltaTime);
 
-	virtual FCursorReply OnCursorQuery( const FGeometry& MyGeometry, const FPointerEvent& CursorEvent ) const override
-	{
-		// Cursor is changed to the dropper when dropper mode is active and the states are correct
-		if (HasMouseCapture() &&
-			bWasClickActivated && bWasLeft && !bWasReEntered )
-		{
-			return FCursorReply::Cursor( EMouseCursor::EyeDropper );
-		}
+	EVisibility GetEscapeTextVisibility() const;
 
-		return SButton::OnCursorQuery( MyGeometry, CursorEvent );
-	}
+	FSlateColor GetDropperImageColor() const;
 
-	FReply OnClicked()
-	{
-		// Log clicks so that OnMouseButtonUp() can post process clicks
-		bWasClicked = true;
-		return FReply::Handled();
-	}
+	bool IsMouseOver(const FGeometry& MyGeometry, const FPointerEvent& MouseEvent) const;
 
 private:
 
-	EVisibility GetEscapeTextVisibility() const
-	{
-		// Show the Esc key message in the button when dropper mode is active
-		if (HasMouseCapture() && bWasClickActivated)
-		{
-			return EVisibility::Visible;
-		}
-		return EVisibility::Hidden;
-	}
-
-	FSlateColor GetDropperImageColor() const
-	{
-		// Make the dropper image in the button pale when dropper mode is active
-		if (HasMouseCapture() && bWasClickActivated)
-		{
-			return FLinearColor(0.3f, 0.3f, 0.3f, 1.0f);
-		}
-		return FSlateColor::UseForeground();
-	}
-
-	bool IsMouseOver(const FGeometry& MyGeometry, const FPointerEvent& MouseEvent) const
-	{ 
-		return MyGeometry.IsUnderLocation(MouseEvent.GetScreenSpacePosition());
-	}
-
 	/** Invoked when a new value is selected by the dropper */
 	FOnLinearColorValueChanged OnValueChanged;
 
diff --git a/Engine/Source/Runtime/AppFramework/Public/Widgets/Colors/SColorPicker.h b/Engine/Source/Runtime/AppFramework/Public/Widgets/Colors/SColorPicker.h
--- a/Engine/Source/Runtime/AppFramework/Public/Widgets/Colors/SColorPicker.h
+++ b/Engine/Source/Runtime/AppFramework/Public/Widgets/Colors/SColorPicker.h
@@ -83,6 +83,7 @@ public:
 		, _OnInteractivePickEnd()
 		, _ParentWindow()
 		, _DisplayGamma(2.2f)
+		, _sRGBOverride()
 		, _DisplayInlineVersion(false)
 		, _ExpandAdvancedSection(false)
 	{ }
@@ -135,6 +136,9 @@ public:
 		/** Sets the display Gamma setting - used to correct colors sampled from the screen */
 		SLATE_ATTRIBUTE(float, DisplayGamma)
 
+		/** Overrides the checkbox value of the sRGB option. */
+		SLATE_ARGUMENT(TOptional<bool>, sRGBOverride)
+
 		/** If true, this color picker will be a stripped down version of the full color picker */
 		SLATE_ARGUMENT(bool, DisplayInlineVersion)
 
@@ -186,7 +190,7 @@ protected:
 	bool SetNewTargetColorRGB( const FLinearColor& NewValue, bool bForceUpdate = false );
 	bool SetNewTargetColorHSV( const FLinearColor& NewValue, bool bForceUpdate = false );
 
-    void UpdateColorPick();
+	void UpdateColorPick();
 	void UpdateColorPickMouseUp();
 
 	void BeginAnimation(FLinearColor Start, FLinearColor End);
@@ -407,6 +411,9 @@ private:
 	/** Sets the display Gamma setting - used to correct colors sampled from the screen */
 	TAttribute<float> DisplayGamma;
 
+	/** Stores the original sRGB option if this color picker temporarily overrides the global option. */
+	TOptional<bool> OriginalSRGBOption;
+
 	/** True if this color picker is an inline color picker */
 	bool bColorPickerIsInlineVersion;
 
@@ -463,6 +470,9 @@ struct FColorPickerArgs
 	/** The current display gamma used to correct colors picked from the display. */
 	TAttribute<float> DisplayGamma;
 
+	/** If set overrides the global option for the desired setting of sRGB mode. */
+	TOptional<bool> sRGBOverride;
+
 	/** An array of FColors to target. */
 	const TArray<FColor*>* ColorArray;
 
@@ -501,6 +511,7 @@ struct FColorPickerArgs
 		, bOnlyRefreshOnOk(false)
 		, bExpandAdvancedSection(true)
 		, DisplayGamma(2.2f)
+		, sRGBOverride()
 		, ColorArray(nullptr)
 		, LinearColorArray(nullptr)
 		, ColorChannelsArray(nullptr)
diff --git a/Engine/Source/Runtime/Core/Private/Mac/MacPlatformMisc.cpp b/Engine/Source/Runtime/Core/Private/Mac/MacPlatformMisc.cpp
--- a/Engine/Source/Runtime/Core/Private/Mac/MacPlatformMisc.cpp
+++ b/Engine/Source/Runtime/Core/Private/Mac/MacPlatformMisc.cpp
@@ -484,7 +484,7 @@ bool FMacPlatformMisc::ControlScreensaver(EScreenSaverAction Action)
 			}
 			break;
 		}
-    }
+	}
 	
 	return true;
 }
@@ -573,9 +573,9 @@ TArray<uint8> FMacPlatformMisc::GetMacAddress()
 				CFDataGetBytes((CFDataRef)MACAddressAsCFData, CFRangeMake(0, kIOEthernetAddressSize), Result.GetData());
 				break;
 				CFRelease(MACAddressAsCFData);
-            }
+			}
 			IOObjectRelease(ControllerService);
-        }
+		}
 		IOObjectRelease(InterfaceService);
 	}
 	IOObjectRelease(InterfaceIterator);
@@ -749,14 +749,14 @@ void FMacPlatformMisc::ClipboardPaste(class FString& Result)
 
 void FMacPlatformMisc::CreateGuid(FGuid& Result)
 {
-    uuid_t UUID;
+	uuid_t UUID;
 	uuid_generate(UUID);
-    
-    uint32* Values = (uint32*)(&UUID[0]);
-    Result[0] = Values[0];
-    Result[1] = Values[1];
-    Result[2] = Values[2];
-    Result[3] = Values[3];
+	
+	uint32* Values = (uint32*)(&UUID[0]);
+	Result[0] = Values[0];
+	Result[1] = Values[1];
+	Result[2] = Values[2];
+	Result[3] = Values[3];
 }
 
 EAppReturnType::Type FMacPlatformMisc::MessageBoxExt(EAppMsgType::Type MsgType, const TCHAR* Text, const TCHAR* Caption)
@@ -1074,30 +1074,23 @@ void FMacPlatformMisc::LoadPreInitModules()
 	FModuleManager::Get().LoadModule(TEXT("CoreAudio"));
 }
 
-FLinearColor FMacPlatformMisc::GetScreenPixelColor(const FVector2D& InScreenPos, float InGamma)
+FLinearColor FMacPlatformMisc::GetScreenPixelColor(const FVector2D& InScreenPos, float /*InGamma*/)
 {
 	SCOPED_AUTORELEASE_POOL;
 
 	CGImageRef ScreenImage = CGWindowListCreateImage(CGRectMake(InScreenPos.X, InScreenPos.Y, 1, 1), kCGWindowListOptionOnScreenBelowWindow, kCGNullWindowID, kCGWindowImageDefault);
-    NSBitmapImageRep *BitmapRep = [[[NSBitmapImageRep alloc] initWithCGImage: ScreenImage] autorelease];
-    NSImage *Image = [[[NSImage alloc] init] autorelease];
-    [Image addRepresentation: BitmapRep];
-	[Image lockFocus];
-	NSColor* PixelColor = NSReadPixel(NSMakePoint(0.0f, 0.0f));
-	[Image unlockFocus];
+	
+	CGDataProviderRef provider = CGImageGetDataProvider(ScreenImage);
+	NSData* data = (id)CGDataProviderCopyData(provider);
+	[data autorelease];
+	const uint8* bytes = (const uint8*)[data bytes];
+	
+	// Mac colors are gamma corrected in Pow(2.2) space, so do the conversion using the 2.2 to linear conversion.
+	FColor ScreenColor(bytes[2], bytes[1], bytes[0]);
+	FLinearColor ScreenLinearColor = FLinearColor::FromPow22Color(ScreenColor);
 	CGImageRelease(ScreenImage);
 
-	FLinearColor ScreenColor([PixelColor redComponent], [PixelColor greenComponent], [PixelColor blueComponent], 1.0f);
-
-	if (InGamma > 1.0f)
-	{
-		// Correct for render gamma
-		ScreenColor.R = FMath::Pow(ScreenColor.R, InGamma);
-		ScreenColor.G = FMath::Pow(ScreenColor.G, InGamma);
-		ScreenColor.B = FMath::Pow(ScreenColor.B, InGamma);
-	}
-	
-	return ScreenColor;
+	return ScreenLinearColor;
 }
 
 FString FMacPlatformMisc::GetCPUVendor()
@@ -1154,21 +1147,21 @@ int32 FMacPlatformMisc::ConvertSlateYPositionToCocoa(int32 YPosition)
 
 int32 FMacPlatformMisc::ConvertCocoaYPositionToSlate(int32 YPosition)
 {
-    NSArray* AllScreens = [NSScreen screens];
-    NSScreen* PrimaryScreen = (NSScreen*)[AllScreens objectAtIndex: 0];
-    NSRect ScreenFrame = [PrimaryScreen frame];
-    NSRect WholeWorkspace = {{0,0},{0,0}};
-    for(NSScreen* Screen in AllScreens)
-    {
-        if(Screen)
-        {
-            WholeWorkspace = NSUnionRect(WholeWorkspace, [Screen frame]);
-        }
-    }
-    
-    CGFloat const OffsetToPrimary = ((ScreenFrame.origin.y + ScreenFrame.size.height) - (WholeWorkspace.origin.y + WholeWorkspace.size.height));
-    CGFloat const OffsetToWorkspace = (WholeWorkspace.size.height - (YPosition)) + WholeWorkspace.origin.y;
-    return OffsetToWorkspace + OffsetToPrimary;
+	NSArray* AllScreens = [NSScreen screens];
+	NSScreen* PrimaryScreen = (NSScreen*)[AllScreens objectAtIndex: 0];
+	NSRect ScreenFrame = [PrimaryScreen frame];
+	NSRect WholeWorkspace = {{0,0},{0,0}};
+	for(NSScreen* Screen in AllScreens)
+	{
+		if(Screen)
+		{
+			WholeWorkspace = NSUnionRect(WholeWorkspace, [Screen frame]);
+		}
+	}
+	
+	CGFloat const OffsetToPrimary = ((ScreenFrame.origin.y + ScreenFrame.size.height) - (WholeWorkspace.origin.y + WholeWorkspace.size.height));
+	CGFloat const OffsetToWorkspace = (WholeWorkspace.size.height - (YPosition)) + WholeWorkspace.origin.y;
+	return OffsetToWorkspace + OffsetToPrimary;
 }
 
 
diff --git a/Engine/Source/Runtime/Core/Private/Windows/WindowsPlatformMisc.cpp b/Engine/Source/Runtime/Core/Private/Windows/WindowsPlatformMisc.cpp
--- a/Engine/Source/Runtime/Core/Private/Windows/WindowsPlatformMisc.cpp
+++ b/Engine/Source/Runtime/Core/Private/Windows/WindowsPlatformMisc.cpp
@@ -1985,27 +1985,18 @@ void FWindowsPlatformMisc::PromptForRemoteDebugging(bool bIsEnsure)
 }
 #endif	//#if !UE_BUILD_SHIPPING
 
-FLinearColor FWindowsPlatformMisc::GetScreenPixelColor(const FVector2D& InScreenPos, float InGamma)
+FLinearColor FWindowsPlatformMisc::GetScreenPixelColor(const FVector2D& InScreenPos, float /*InGamma*/)
 {
 	COLORREF PixelColorRef = GetPixel(GetDC(HWND_DESKTOP), InScreenPos.X, InScreenPos.Y);
 
-	const float DivideBy255 = 1.0f / 255.0f;
+	FColor sRGBScreenColor(
+		(PixelColorRef & 0xFF),
+		((PixelColorRef & 0xFF00) >> 8),
+		((PixelColorRef & 0xFF0000) >> 16),
+		255);
 
-	FLinearColor ScreenColor(
-		(PixelColorRef & 0xFF) * DivideBy255,
-		((PixelColorRef & 0xFF00) >> 8) * DivideBy255,
-		((PixelColorRef & 0xFF0000) >> 16) * DivideBy255,
-		1.0f);
-
-	if (InGamma > 1.0f)
-	{
-		// Correct for render gamma
-		ScreenColor.R = FMath::Pow(ScreenColor.R, InGamma);
-		ScreenColor.G = FMath::Pow(ScreenColor.G, InGamma);
-		ScreenColor.B = FMath::Pow(ScreenColor.B, InGamma);
-	}
-
-	return ScreenColor;
+	// Assume the screen color is coming in as sRGB space
+	return FLinearColor(sRGBScreenColor);
 }
 
 /**
