UMG/Slate - Attempting to fix an issue in the mac build's AppFramework, needed to link in CoreUObject

diff --git a/Engine/Source/Runtime/AppFramework/AppFramework.Build.cs b/Engine/Source/Runtime/AppFramework/AppFramework.Build.cs
--- a/Engine/Source/Runtime/AppFramework/AppFramework.Build.cs
+++ b/Engine/Source/Runtime/AppFramework/AppFramework.Build.cs
@@ -9,6 +9,7 @@ public class AppFramework : ModuleRules
 		PublicDependencyModuleNames.AddRange(
 			new string[] {
 				"Core",
+                "CoreUObject",
 				"Slate",
 				"SlateCore",
 			}
