Added sRGB color input support in ColorPicker (GitHub PR# 1436 by maxpestun); some layout tweaks.

https://github.com/EpicGames/UnrealEngine/pull/1436

diff --git a/Engine/Source/Runtime/AppFramework/Private/Widgets/Colors/SColorPicker.cpp b/Engine/Source/Runtime/AppFramework/Private/Widgets/Colors/SColorPicker.cpp
--- a/Engine/Source/Runtime/AppFramework/Private/Widgets/Colors/SColorPicker.cpp
+++ b/Engine/Source/Runtime/AppFramework/Private/Widgets/Colors/SColorPicker.cpp
@@ -319,7 +319,7 @@ void SColorPicker::GenerateDefaultColorPickerContent( bool bAdvancedSectionExpan
 				.OnCheckStateChanged(this, &SColorPicker::HandleSRGBCheckBoxCheckStateChanged)
 				[
 					SNew(STextBlock)
-					.Text(LOCTEXT("SRGBCheckboxLabel", "sRGB"))
+					.Text(LOCTEXT("SRGBCheckboxLabel", "sRGB Preview"))
 				]
 			]
 
@@ -511,32 +511,60 @@ void SColorPicker::GenerateDefaultColorPickerContent( bool bAdvancedSectionExpan
 							MakeColorSpinBox(EColorPickerChannels::Value)
 						]
 
-						// Hex
+						// Hex linear
 						+ SVerticalBox::Slot()
-						.HAlign(HAlign_Right)
-						.VAlign(VAlign_Top)
-						.Padding(0.0f, 8.0f, 0.0f, 0.0f)
-						[
-							SNew(SHorizontalBox)
-							.ToolTipText(LOCTEXT("HexSliderToolTip", "Hexadecimal Value"))
-
-							+ SHorizontalBox::Slot()
-							.AutoWidth()
-							.Padding(0.0f, 0.0f, 4.0f, 0.0f)
-							.VAlign(VAlign_Center)
+							.HAlign(HAlign_Right)
+							.VAlign(VAlign_Top)
+							.Padding(0.0f, 12.0f, 0.0f, 0.0f)
 							[
-								SNew(STextBlock)
-								.Text(LOCTEXT("HexInputLabel", "Hex"))
-							]
+								SNew(SHorizontalBox)
+								.ToolTipText(LOCTEXT("HexLinearSliderToolTip", "Hexadecimal Linear Value"))
 
-							+ SHorizontalBox::Slot()
-							.AutoWidth()
+								+ SHorizontalBox::Slot()
+								.AutoWidth()
+								.Padding(0.0f, 0.0f, 4.0f, 0.0f)
+								.VAlign(VAlign_Center)
+								[
+									SNew(STextBlock)
+									.Text(LOCTEXT("HexLinearInputLabel", "Hex Linear"))
+								]
+
+								+ SHorizontalBox::Slot()
+									.AutoWidth()
+									[
+										SNew(SEditableTextBox)
+										.MinDesiredWidth(72.0f)
+										.Text(this, &SColorPicker::HandleHexLinearBoxText)
+										.OnTextCommitted(this, &SColorPicker::HandleHexLinearInputTextCommitted)
+									]
+						]
+
+						// Hex sRGB
+						+ SVerticalBox::Slot()
+							.HAlign(HAlign_Right)
+							.VAlign(VAlign_Top)
+							.Padding(0.0f, 8.0f, 0.0f, 0.0f)
 							[
-								SNew(SEditableTextBox)
-								.MinDesiredWidth(72.0f)
-								.Text(this, &SColorPicker::HandleHexBoxText)
-								.OnTextCommitted(this, &SColorPicker::HandleHexInputTextCommitted)
-							]
+								SNew(SHorizontalBox)
+								.ToolTipText(LOCTEXT("HexSRGBSliderToolTip", "Hexadecimal sRGB Value"))
+
+								+ SHorizontalBox::Slot()
+								.AutoWidth()
+								.Padding(0.0f, 0.0f, 4.0f, 0.0f)
+								.VAlign(VAlign_Center)
+								[
+									SNew(STextBlock)
+									.Text(LOCTEXT("HexSRGBInputLabel", "Hex sRGB"))
+								]
+
+								+ SHorizontalBox::Slot()
+									.AutoWidth()
+									[
+										SNew(SEditableTextBox)
+										.MinDesiredWidth(72.0f)
+										.Text(this, &SColorPicker::HandleHexSRGBBoxText)
+										.OnTextCommitted(this, &SColorPicker::HandleHexSRGBInputTextCommitted)
+									]
 						]
 					]
 				]
@@ -547,6 +575,7 @@ void SColorPicker::GenerateDefaultColorPickerContent( bool bAdvancedSectionExpan
 			.AutoHeight()
 			.HAlign(HAlign_Right)
 			.VAlign(VAlign_Center)
+			.Padding(0.0f, 12.0f, 0.0f, 0.0f)
 			[
 				SNew(SUniformGridPanel)
 					.MinDesiredSlotHeight(FCoreStyle::Get().GetFloat("StandardDialog.MinDesiredSlotHeight"))
@@ -1263,13 +1292,19 @@ void SColorPicker::HandleEyeDropperButtonComplete()
 }
 
 
-FText SColorPicker::HandleHexBoxText() const
+FText SColorPicker::HandleHexLinearBoxText() const
 {
 	return FText::FromString(CurrentColorRGB.ToFColor(false).ToHex());
 }
 
 
-void SColorPicker::HandleHexInputTextCommitted( const FText& Text, ETextCommit::Type CommitType )
+FText SColorPicker::HandleHexSRGBBoxText() const
+{
+	return FText::FromString(CurrentColorRGB.ToFColor(true).ToHex());
+}
+
+
+void SColorPicker::HandleHexLinearInputTextCommitted(const FText& Text, ETextCommit::Type CommitType)
 {
 	if (!Text.IsEmpty() && ((CommitType == ETextCommit::OnEnter) || (CommitType == ETextCommit::OnUserMovedFocus)))
 	{
@@ -1278,6 +1313,24 @@ void SColorPicker::HandleHexInputTextCommitted( const FText& Text, ETextCommit::
 	}	
 }
 
+void SColorPicker::HandleHexSRGBInputTextCommitted(const FText& Text, ETextCommit::Type CommitType)
+{
+	if (!Text.IsEmpty() && ((CommitType == ETextCommit::OnEnter) || (CommitType == ETextCommit::OnUserMovedFocus)))
+	{
+		FColor Color = FColor::FromHex(Text.ToString());
+		float red = Color.R / 255.0f;
+		float green = Color.G / 255.0f;
+		float blue = Color.B / 255.0f;
+		float alpha = Color.A / 255.0f;
+
+		red = red <= 0.04045f ? red / 12.92f : FMath::Pow((red + 0.055f) / 1.055f, 2.4f);
+		green = green <= 0.04045f ? green / 12.92f : FMath::Pow((green + 0.055f) / 1.055f, 2.4f);
+		blue = blue <= 0.04045f ? blue / 12.92f : FMath::Pow((blue + 0.055f) / 1.055f, 2.4f);
+		
+		SetNewTargetColorRGB(FLinearColor(red, green, blue, alpha), false);
+	}
+}
+
 
 void SColorPicker::HandleHSVColorChanged( FLinearColor NewValue )
 {
diff --git a/Engine/Source/Runtime/AppFramework/Private/Widgets/Colors/SColorThemes.cpp b/Engine/Source/Runtime/AppFramework/Private/Widgets/Colors/SColorThemes.cpp
--- a/Engine/Source/Runtime/AppFramework/Private/Widgets/Colors/SColorThemes.cpp
+++ b/Engine/Source/Runtime/AppFramework/Private/Widgets/Colors/SColorThemes.cpp
@@ -746,26 +746,22 @@ FText SColorThemeBar::GetThemeName() const
 
 FReply SColorThemeBar::OnMouseButtonDown( const FGeometry& MyGeometry, const FPointerEvent& MouseEvent )
 {
-	if ( MouseEvent.GetEffectingButton() == EKeys::LeftMouseButton )
+	if (MouseEvent.GetEffectingButton() == EKeys::LeftMouseButton)
 	{
 		OnCurrentThemeChanged.ExecuteIfBound(ColorTheme.Pin());
 
 		return FReply::Handled();
 	}
-	else
-	{
-		return FReply::Unhandled();
-	}
-}
-
-
 
+	return FReply::Unhandled();
+}
 
 
 TArray< TSharedPtr<FColorTheme> > SColorThemesViewer::ColorThemes;
 TWeakPtr<FColorTheme> SColorThemesViewer::CurrentlySelectedThemePtr;
 bool SColorThemesViewer::bSRGBEnabled = false;
 
+
 BEGIN_SLATE_FUNCTION_BUILD_OPTIMIZATION
 void SColorThemesViewer::Construct(const FArguments& InArgs)
 {
diff --git a/Engine/Source/Runtime/AppFramework/Public/Widgets/Colors/SColorPicker.h b/Engine/Source/Runtime/AppFramework/Public/Widgets/Colors/SColorPicker.h
--- a/Engine/Source/Runtime/AppFramework/Public/Widgets/Colors/SColorPicker.h
+++ b/Engine/Source/Runtime/AppFramework/Public/Widgets/Colors/SColorPicker.h
@@ -271,11 +271,17 @@ private:
 	// Callback for completed eye dropper interactions.
 	void HandleEyeDropperButtonComplete();
 
-	// Callback for getting the text in the hex box.
-	FText HandleHexBoxText() const;
+	// Callback for getting the text in the hex linear box.
+	FText HandleHexLinearBoxText() const;
 
-	// Callback for committed text in the hex input box.
-	void HandleHexInputTextCommitted( const FText& Text, ETextCommit::Type CommitType );
+	// Callback for getting the text in the hex sRGB box.
+	FText HandleHexSRGBBoxText() const;
+
+	// Callback for committed text in the hex input box (sRGB gamma).
+	void HandleHexSRGBInputTextCommitted(const FText& Text, ETextCommit::Type CommitType);
+
+	// Callback for committed text in the hex input box (linear gamma).
+	void HandleHexLinearInputTextCommitted(const FText& Text, ETextCommit::Type CommitType);
 
 	// Callback for changing the HSV value of the current color.
 	void HandleHSVColorChanged( FLinearColor NewValue );
