VREditor: Color Picker in VR implemented with delegate overrides for creation and destruction. Transforms still need adjustment.

diff --git a/Engine/Source/Editor/VREditor/VREditor.Build.cs b/Engine/Source/Editor/VREditor/VREditor.Build.cs
--- a/Engine/Source/Editor/VREditor/VREditor.Build.cs
+++ b/Engine/Source/Editor/VREditor/VREditor.Build.cs
@@ -8,6 +8,7 @@ namespace UnrealBuildTool.Rules
 		{
 			PrivateDependencyModuleNames.AddRange(
                 new string[] {
+                    "AppFramework",
 				    "Core",
 				    "CoreUObject",
 				    "Engine",
diff --git a/Engine/Source/Editor/VREditor/VREditorUISystem.cpp b/Engine/Source/Editor/VREditor/VREditorUISystem.cpp
--- a/Engine/Source/Editor/VREditor/VREditorUISystem.cpp
+++ b/Engine/Source/Editor/VREditor/VREditorUISystem.cpp
@@ -37,6 +37,7 @@
 #include "SDPIScaler.h"
 #include "SWidget.h"
 #include "SDockTab.h"
+#include "SColorPicker.h"
 
 namespace VREd
 {
@@ -145,6 +146,10 @@ void UVREditorUISystem::Init()
 	
 	// Create all of our UI panels
 	CreateUIs();
+
+	// Bind the color picker creation & destruction overrides
+	SColorPicker::OnColorPickerNonModalCreateOverride.BindUObject(this, &UVREditorUISystem::CreateVRColorPicker);
+	SColorPicker::OnColorPickerDestroyOverride.BindUObject(this, &UVREditorUISystem::DestroyVRColorPicker);
 }
 
 
@@ -156,10 +161,15 @@ void UVREditorUISystem::Shutdown()
 		GetOwner().GetWorldInteraction().OnViewportInteractionHoverUpdate().RemoveAll( this );
 	}
 
+	// Unbind the color picker creation & destruction overrides
+	SColorPicker::OnColorPickerNonModalCreateOverride.Unbind();
+	SColorPicker::OnColorPickerDestroyOverride.Unbind();
+
 	CleanUpActorsBeforeMapChangeOrSimulate();
 
 	DefaultWindowTransforms.Empty();
 
+
 	QuickMenuWidgetClass = nullptr;
 	QuickRadialWidgetClass = nullptr;
 	TutorialWidgetClass = nullptr;
@@ -1129,19 +1139,18 @@ void UVREditorUISystem::StartDraggingDockUI( AVREditorDockableWindow* InitDraggi
 	if( DockTo == AVREditorFloatingUI::EDockedTo::LeftHand || DockTo == AVREditorFloatingUI::EDockedTo::RightHand )
 	{
 		bDraggedDockFromHandPassedThreshold = false;
-	}
-
-	// Set the laser of the interactor the dockable window was pulled from
-	UViewportInteractor* OtherInteractor = Interactor->GetOtherInteractor();
-	if ( OtherInteractor )
-	{
-		UVREditorInteractor* OtherVREditorInteractor = Cast<UVREditorInteractor>( OtherInteractor );
-		if ( OtherVREditorInteractor )
+	
+		// Set the laser of the interactor the dockable window was pulled from
+		UViewportInteractor* OtherInteractor = Interactor->GetOtherInteractor();
+		if ( OtherInteractor )
 		{
-			OtherVREditorInteractor->SetHasUIInFront( false );
+			UVREditorInteractor* OtherVREditorInteractor = Cast<UVREditorInteractor>( OtherInteractor );
+			if ( OtherVREditorInteractor )
+			{
+				OtherVREditorInteractor->SetHasUIInFront( false );
+			}
 		}
 	}
-
 	InteractorDraggingUI = Interactor;
 
 	FTransform UIToWorld = InitDraggingDockUI->GetActorTransform();
@@ -1394,4 +1403,67 @@ void UVREditorUISystem::SetDefaultWindowLayout()
 			Panel->SetRelativeOffset( RoomSpaceHeadLocation + NewTransform.GetTranslation() );
 		}
 	}
-}
\ No newline at end of file
+}
+
+void UVREditorUISystem::CreateVRColorPicker(const TSharedRef<SColorPicker>& ColorPicker)
+{
+	DestroyVRColorPicker();
+	
+	// Check that the Color Picker Panel isn't currently showing. Also handles the EEditorUIPanel::ColorPicker being null.
+	if(!IsShowingEditorUIPanel( EEditorUIPanel::ColorPicker ))
+	{
+		const FIntPoint DefaultResolution( VREd::DefaultEditorUIResolutionX->GetInt(), VREd::DefaultEditorUIResolutionY->GetInt() );
+
+		TSharedRef<SWidget> WidgetToDraw =
+			SNew(SDPIScaler)
+			.DPIScale(VREd::EditorUIScale->GetFloat())
+			[
+				ColorPicker
+			]
+		;
+
+		const bool bWithSceneComponent = false;
+		AVREditorFloatingUI* ColorPickerUI = GetOwner().SpawnTransientSceneActor< AVREditorDockableWindow >(TEXT("ColorPickerUI"), bWithSceneComponent);
+		ColorPickerUI->ShowUI(true);
+		FloatingUIs.Add(ColorPickerUI);
+
+		EditorUIPanels[ (int32)EEditorUIPanel::ColorPicker ] = ColorPickerUI;
+	
+		// Always spawn based on the location of the menu you are hovering over. To get this information, find the hand hovering over the UI.
+		UVREditorInteractor* VREditorInteractor = GetOwner().GetHandInteractor(EControllerHand::Left); 
+		if (!VREditorInteractor->IsHoveringOverUI())
+		{
+			VREditorInteractor = GetOwner().GetHandInteractor(EControllerHand::Right);
+
+		}
+	
+		ColorPickerUI->SetSlateWidget(*this, WidgetToDraw, DefaultResolution, VREd::EditorUISize->GetFloat(), AVREditorFloatingUI::EDockedTo::Room);
+			
+		FVector RelativeLocation = VREditorInteractor->GetInteractorData().HoverLocation;
+		// Adjust location for room transform, world scale
+		RelativeLocation -= GetOwner().GetRoomTransform().GetLocation();
+		RelativeLocation /= GetOwner().GetWorldScaleFactor();
+		ColorPickerUI->SetRelativeOffset(RelativeLocation);
+
+
+		FRotator RelativeRotation = VREditorInteractor->GetHitResultFromLaserPointer().GetActor()->GetActorRotation();
+		//Use this forward vector for the offset: RelativeRotation.RotateVector(FVector::ForwardVector);
+		//RelativeRotation -= (GetOwner().GetRoomTransform().GetRotation().Rotator());
+		ColorPickerUI->SetLocalRotation(RelativeRotation);
+
+		const bool bShouldShow = true;
+		const bool bShowOnHand = false;
+		ShowEditorUIPanel(EEditorUIPanel::ColorPicker, VREditorInteractor, bShouldShow, bShowOnHand);
+	}
+	
+}
+
+void UVREditorUISystem::DestroyVRColorPicker()
+{
+
+	UVREditorInteractor* VREditorInteractor = nullptr;
+	const bool bShouldShow = false;
+	const bool bShowOnHand = false;
+	const bool bRefreshQuickMenu = true;
+	ShowEditorUIPanel(EEditorUIPanel::ColorPicker, VREditorInteractor, bShouldShow, bShowOnHand, bRefreshQuickMenu);
+}
diff --git a/Engine/Source/Editor/VREditor/VREditorUISystem.h b/Engine/Source/Editor/VREditor/VREditorUISystem.h
--- a/Engine/Source/Editor/VREditor/VREditorUISystem.h
+++ b/Engine/Source/Editor/VREditor/VREditorUISystem.h
@@ -6,6 +6,7 @@
 
 class UVREditorInteractor;
 class UViewportInteractor;
+class SColorPicker;
 
 /**
  * VR Editor user interface manager
@@ -26,6 +27,7 @@ public:
 		Tutorial,
 		AssetEditor,
 		WorldSettings,
+		ColorPicker,
 		// ...
 
 		TotalCount,
@@ -150,6 +152,12 @@ protected:
 	/** Sets the main windows to their default transform */
 	void SetDefaultWindowLayout();
 
+	/** Creates a VR-specific color picker. Gets bound to SColorPicker's creation override delegate */
+	void CreateVRColorPicker(const TSharedRef<SColorPicker>& ColorPicker);
+
+	/** Hides the VR-specific color picker. Gets bound to SColorPicker's destruction override delegate */
+	void DestroyVRColorPicker();
+
 protected:
 
 	/** Owning object */
diff --git a/Engine/Source/Runtime/AppFramework/Private/Widgets/Colors/SColorPicker.cpp b/Engine/Source/Runtime/AppFramework/Private/Widgets/Colors/SColorPicker.cpp
--- a/Engine/Source/Runtime/AppFramework/Private/Widgets/Colors/SColorPicker.cpp
+++ b/Engine/Source/Runtime/AppFramework/Private/Widgets/Colors/SColorPicker.cpp
@@ -62,6 +62,7 @@ void SColorPicker::Construct( const FArguments& InArgs )
 	ParentWindowPtr = InArgs._ParentWindow.Get();
 	DisplayGamma = InArgs._DisplayGamma;
 	bClosedViaOkOrCancel = false;
+	bValidCreationOverrideExists = InArgs._OverrideColorPickerCreation;
 
 	if ( InArgs._sRGBOverride.IsSet() )
 	{
@@ -80,6 +81,7 @@ void SColorPicker::Construct( const FArguments& InArgs )
 	bColorPickerIsInlineVersion = InArgs._DisplayInlineVersion;
 	bIsInteractive = false;
 	bPerfIsTooSlowToUpdate = false;
+	
 
 	BackupColors();
 
@@ -581,7 +583,7 @@ void SColorPicker::GenerateDefaultColorPickerContent( bool bAdvancedSectionExpan
 					.MinDesiredSlotHeight(FCoreStyle::Get().GetFloat("StandardDialog.MinDesiredSlotHeight"))
 					.MinDesiredSlotWidth(FCoreStyle::Get().GetFloat("StandardDialog.MinDesiredSlotWidth"))
 					.SlotPadding(FCoreStyle::Get().GetMargin("StandardDialog.SlotPadding"))
-					.Visibility(ParentWindowPtr.IsValid() ? EVisibility::Visible : EVisibility::Collapsed)
+					.Visibility((ParentWindowPtr.IsValid() || bValidCreationOverrideExists ) ? EVisibility::Visible : EVisibility::Collapsed)
 
 				+ SUniformGridPanel::Slot(0, 0)
 					[
@@ -1189,7 +1191,16 @@ FReply SColorPicker::HandleCancelButtonClicked()
 	bClosedViaOkOrCancel = true;
 
 	DiscardColor();
-	ParentWindowPtr.Pin()->RequestDestroyWindow();
+	
+
+	if (SColorPicker::OnColorPickerDestroyOverride.IsBound())
+	{
+		SColorPicker::OnColorPickerDestroyOverride.Execute();
+	}
+	else
+	{
+		ParentWindowPtr.Pin()->RequestDestroyWindow();
+	}
 
 	return FReply::Handled();
 }
@@ -1437,7 +1448,16 @@ FReply SColorPicker::HandleOkButtonClicked()
 		UpdateColorPick();
 	}
 
-	ParentWindowPtr.Pin()->RequestDestroyWindow();
+	if (SColorPicker::OnColorPickerDestroyOverride.IsBound())
+	{
+		SColorPicker::OnColorPickerDestroyOverride.Execute();
+	}
+	else
+	{
+		ParentWindowPtr.Pin()->RequestDestroyWindow();
+	}
+
+	
 
 	return FReply::Handled();
 }
@@ -1581,6 +1601,9 @@ void SColorPicker::HandleThemesViewerThemeChanged()
 	}
 }
 
+// Static delegates to access whether or not the override is bound in the global Open/Destroy functions
+SColorPicker::FOnColorPickerCreationOverride SColorPicker::OnColorPickerNonModalCreateOverride;
+SColorPicker::FOnColorPickerDestructionOverride SColorPicker::OnColorPickerDestroyOverride;
 
 /* Global functions
  *****************************************************************************/
@@ -1625,13 +1648,21 @@ bool OpenColorPicker(const FColorPickerArgs& Args)
 
 	FVector2D AdjustedSummonLocation = FSlateApplication::Get().CalculatePopupWindowPosition( Anchor, SColorPicker::DEFAULT_WINDOW_SIZE, Orient_Horizontal );
 
-	TSharedPtr<SWindow> Window = SNew(SWindow)
-		.AutoCenter(EAutoCenter::None)
-		.ScreenPosition(AdjustedSummonLocation)
-		.SupportsMaximize(false)
-		.SupportsMinimize(false)
-		.SizingRule(ESizingRule::Autosized)
-		.Title(LOCTEXT("WindowHeader", "Color Picker"));
+	// Only override the color picker window creation behavior if we are not creating a modal color picker
+	const bool bOverrideNonModalCreation = (SColorPicker::OnColorPickerNonModalCreateOverride.IsBound() && !Args.bIsModal);
+
+	TSharedPtr<SWindow> Window = nullptr;
+	
+	if(!bOverrideNonModalCreation)
+	{
+		Window = SNew(SWindow)
+			.AutoCenter(EAutoCenter::None)
+			.ScreenPosition(AdjustedSummonLocation)
+			.SupportsMaximize(false)
+			.SupportsMinimize(false)
+			.SizingRule(ESizingRule::Autosized)
+			.Title(LOCTEXT("WindowHeader", "Color Picker"));
+	}
 
 	TSharedRef<SColorPicker> ColorPicker = SNew(SColorPicker)
 		.TargetColorAttribute(OldColor)
@@ -1650,45 +1681,60 @@ bool OpenColorPicker(const FColorPickerArgs& Args)
 		.OnColorPickerWindowClosed(Args.OnColorPickerWindowClosed)
 		.ParentWindow(Window)
 		.DisplayGamma(Args.DisplayGamma)
-		.sRGBOverride(Args.sRGBOverride);
-		
-	Window->SetContent(
-		SNew(SBox)
-		[
-			SNew(SBorder)
-			.BorderImage(FCoreStyle::Get().GetBrush("ToolPanel.GroupBorder"))
-			.Padding(FMargin(8.0f, 8.0f))
-			[
-				ColorPicker
-			]
-		]
-	);
+		.sRGBOverride(Args.sRGBOverride)
+		.OverrideColorPickerCreation(bOverrideNonModalCreation);
+
 
-	if (Args.bIsModal)
+
+	// If the color picker requested is modal, don't override the behavior even if the delegate is bound
+	if (bOverrideNonModalCreation)
 	{
-		FSlateApplication::Get().AddModalWindow(Window.ToSharedRef(), Args.ParentWidget);
+		SColorPicker::OnColorPickerNonModalCreateOverride.Execute(ColorPicker);
+
+		Result = true;
+
+		//hold on to the window created for external use...
+		ColorPickerWindow = Window;
 	}
 	else
 	{
-		if ( Args.ParentWidget.IsValid() )
+		Window->SetContent(
+			SNew(SBox)
+			[
+				SNew(SBorder)
+				.BorderImage(FCoreStyle::Get().GetBrush("ToolPanel.GroupBorder"))
+			.Padding(FMargin(8.0f, 8.0f))
+			[
+				ColorPicker
+			]
+			]
+		);
+
+		if (Args.bIsModal)
 		{
-			// Find the window of the parent widget
-			FWidgetPath WidgetPath;
-			FSlateApplication::Get().GeneratePathToWidgetChecked( Args.ParentWidget.ToSharedRef(), WidgetPath );
-			Window = FSlateApplication::Get().AddWindowAsNativeChild( Window.ToSharedRef(), WidgetPath.GetWindow() );
+			FSlateApplication::Get().AddModalWindow(Window.ToSharedRef(), Args.ParentWidget);
 		}
 		else
 		{
-			Window = FSlateApplication::Get().AddWindow(Window.ToSharedRef());
+			if (Args.ParentWidget.IsValid())
+			{
+				// Find the window of the parent widget
+				FWidgetPath WidgetPath;
+				FSlateApplication::Get().GeneratePathToWidgetChecked(Args.ParentWidget.ToSharedRef(), WidgetPath);
+				Window = FSlateApplication::Get().AddWindowAsNativeChild(Window.ToSharedRef(), WidgetPath.GetWindow());
+			}
+			else
+			{
+				Window = FSlateApplication::Get().AddWindow(Window.ToSharedRef());
+			}
+
 		}
-			
-	}
 
-	Result = true;
+		Result = true;
 
-	//hold on to the window created for external use...
-	ColorPickerWindow = Window;
-	
+		//hold on to the window created for external use...
+		ColorPickerWindow = Window;
+	}
 #endif
 
 	return Result;
@@ -1703,7 +1749,14 @@ void DestroyColorPicker()
 {
 	if (ColorPickerWindow.IsValid())
 	{
-		ColorPickerWindow.Pin()->RequestDestroyWindow();
+		if (SColorPicker::OnColorPickerDestroyOverride.IsBound())
+		{
+			SColorPicker::OnColorPickerDestroyOverride.Execute();
+		}
+		else
+		{
+			ColorPickerWindow.Pin()->RequestDestroyWindow();
+		}
 		ColorPickerWindow.Reset();
 	}
 }
diff --git a/Engine/Source/Runtime/AppFramework/Public/Widgets/Colors/SColorPicker.h b/Engine/Source/Runtime/AppFramework/Public/Widgets/Colors/SColorPicker.h
--- a/Engine/Source/Runtime/AppFramework/Public/Widgets/Colors/SColorPicker.h
+++ b/Engine/Source/Runtime/AppFramework/Public/Widgets/Colors/SColorPicker.h
@@ -10,11 +10,9 @@ class SColorThemesViewer;
 class SComboButton;
 class SThemeColorBlocksBar;
 
-
 /** Called when the color picker cancel button is pressed */
 DECLARE_DELEGATE_OneParam(FOnColorPickerCancelled, FLinearColor);
 
-
 /**
  * Enumerates color channels (do not reorder).
  */
@@ -85,6 +83,7 @@ public:
 		, _DisplayGamma(2.2f)
 		, _sRGBOverride()
 		, _DisplayInlineVersion(false)
+		, _OverrideColorPickerCreation(false)
 		, _ExpandAdvancedSection(false)
 	{ }
 		
@@ -142,6 +141,9 @@ public:
 		/** If true, this color picker will be a stripped down version of the full color picker */
 		SLATE_ARGUMENT(bool, DisplayInlineVersion)
 
+		/** If true, this color picker will have non-standard creation behavior */
+		SLATE_ARGUMENT(bool, OverrideColorPickerCreation)
+
 		/** If true, the Advanced section will be expanded, regardless of the remembered state */
 		SLATE_ARGUMENT(bool, ExpandAdvancedSection)
 
@@ -162,6 +164,13 @@ public:
 	 */
 	void Construct(const FArguments& InArgs );
 
+	/** Delegate to override color picker creation behavior */
+	DECLARE_DELEGATE_OneParam(FOnColorPickerCreationOverride, const TSharedRef<SColorPicker>&);
+	static FOnColorPickerCreationOverride OnColorPickerNonModalCreateOverride;
+	/** Delegate to override color picker destruction behavior */
+	DECLARE_DELEGATE(FOnColorPickerDestructionOverride);
+	static FOnColorPickerDestructionOverride OnColorPickerDestroyOverride;
+
 protected:
 
 	/** Backup all the colors that are being modified */
@@ -423,9 +432,15 @@ private:
 	/** True if this color picker is an inline color picker */
 	bool bColorPickerIsInlineVersion;
 
+	/** True if something has overridden the color picker's creation behavior */
+	bool bColorPickerCreationIsOverridden;
+
 	/** Tracks whether the user is moving a value spin box, the color wheel and the dropper */
 	bool bIsInteractive;
 
+	/** Is true if the color picker creation behavior can be overridden */
+	bool bValidCreationOverrideExists;
+
 private:
 
 	/** Invoked when a new value is selected on the color wheel */
@@ -437,6 +452,8 @@ private:
 	/** Invoked when the color picker cancel button is pressed */
 	FOnColorPickerCancelled OnColorPickerCancelled;
 
+
+
 	/** Invoked when a slider drag, color wheel drag or dropper grab starts */
 	FSimpleDelegate OnInteractivePickBegin;
 
