Duplicating color themes now dupllicates the color labels as well

diff --git a/Engine/Source/Runtime/AppFramework/Private/Widgets/Colors/SColorThemes.cpp b/Engine/Source/Runtime/AppFramework/Private/Widgets/Colors/SColorThemes.cpp
--- a/Engine/Source/Runtime/AppFramework/Private/Widgets/Colors/SColorThemes.cpp
+++ b/Engine/Source/Runtime/AppFramework/Private/Widgets/Colors/SColorThemes.cpp
@@ -1140,7 +1140,7 @@ FReply SColorThemesViewer::DuplicateColorTheme()
 	const TArray< TSharedPtr<FColorInfo> >& CurrentColors = CurrentlySelectedThemePtr.Pin()->GetColors();
 	for (int32 ColorIndex = 0; ColorIndex < CurrentColors.Num(); ++ColorIndex)
 	{
-		NewColors.Add(MakeShareable(new FColorInfo(CurrentColors[ColorIndex]->Color)));
+		NewColors.Add(MakeShareable(new FColorInfo(CurrentColors[ColorIndex]->Color, CurrentColors[ColorIndex]->Label)));
 	}
 	const FText Name = NSLOCTEXT("ColorThemesViewer", "CopyThemeNameAppend", " Copy");
 	NewColorTheme( CurrentlySelectedThemePtr.Pin()->Name + Name.ToString(), NewColors );
diff --git a/Engine/Source/Runtime/AppFramework/Public/Widgets/Colors/SColorThemes.h b/Engine/Source/Runtime/AppFramework/Public/Widgets/Colors/SColorThemes.h
--- a/Engine/Source/Runtime/AppFramework/Public/Widgets/Colors/SColorThemes.h
+++ b/Engine/Source/Runtime/AppFramework/Public/Widgets/Colors/SColorThemes.h
@@ -32,6 +32,12 @@ struct FColorInfo
 		Color = InColor;
 		Label = FText();
 	}
+
+	FColorInfo(TSharedPtr<FLinearColor> InColor, FText InLabel)
+	{
+		Color = InColor;
+		Label = InLabel;
+	}
 };
 
 /**
