Fixes for mismatched macros.

diff --git a/Engine/Source/Developer/Apple/MetalShaderFormat/Private/MetalBackend.cpp b/Engine/Source/Developer/Apple/MetalShaderFormat/Private/MetalBackend.cpp
--- a/Engine/Source/Developer/Apple/MetalShaderFormat/Private/MetalBackend.cpp
+++ b/Engine/Source/Developer/Apple/MetalShaderFormat/Private/MetalBackend.cpp
@@ -12,7 +12,7 @@
 
 PRAGMA_DISABLE_SHADOW_VARIABLE_WARNINGS
 #include "glsl_parser_extras.h"
-PRAGMA_POP
+PRAGMA_ENABLE_SHADOW_VARIABLE_WARNINGS
 
 #include "hash_table.h"
 #include "ir_rvalue_visitor.h"
diff --git a/Engine/Source/Developer/ShaderFormatVectorVM/Private/VectorVMBackend.cpp b/Engine/Source/Developer/ShaderFormatVectorVM/Private/VectorVMBackend.cpp
--- a/Engine/Source/Developer/ShaderFormatVectorVM/Private/VectorVMBackend.cpp
+++ b/Engine/Source/Developer/ShaderFormatVectorVM/Private/VectorVMBackend.cpp
@@ -8,7 +8,7 @@
 
 PRAGMA_DISABLE_SHADOW_VARIABLE_WARNINGS
 #include "glsl_parser_extras.h"
-PRAGMA_POP
+PRAGMA_ENABLE_SHADOW_VARIABLE_WARNINGS
 
 #include "hash_table.h"
 #include "ir_rvalue_visitor.h"
diff --git a/Engine/Source/Developer/ShaderFormatVectorVM/Private/ir_vm_flatten_branches_to_selects_visitor.cpp b/Engine/Source/Developer/ShaderFormatVectorVM/Private/ir_vm_flatten_branches_to_selects_visitor.cpp
--- a/Engine/Source/Developer/ShaderFormatVectorVM/Private/ir_vm_flatten_branches_to_selects_visitor.cpp
+++ b/Engine/Source/Developer/ShaderFormatVectorVM/Private/ir_vm_flatten_branches_to_selects_visitor.cpp
@@ -8,7 +8,7 @@
 
 PRAGMA_DISABLE_SHADOW_VARIABLE_WARNINGS
 #include "glsl_parser_extras.h"
-PRAGMA_POP
+PRAGMA_ENABLE_SHADOW_VARIABLE_WARNINGS
 
 #include "hash_table.h"
 #include "ir_rvalue_visitor.h"
diff --git a/Engine/Source/Developer/ShaderFormatVectorVM/Private/ir_vm_gen_bytecode_visitor.cpp b/Engine/Source/Developer/ShaderFormatVectorVM/Private/ir_vm_gen_bytecode_visitor.cpp
--- a/Engine/Source/Developer/ShaderFormatVectorVM/Private/ir_vm_gen_bytecode_visitor.cpp
+++ b/Engine/Source/Developer/ShaderFormatVectorVM/Private/ir_vm_gen_bytecode_visitor.cpp
@@ -9,7 +9,7 @@
 
 PRAGMA_DISABLE_SHADOW_VARIABLE_WARNINGS
 #include "glsl_parser_extras.h"
-PRAGMA_POP
+PRAGMA_ENABLE_SHADOW_VARIABLE_WARNINGS
 
 #include "hash_table.h"
 #include "ir_rvalue_visitor.h"
diff --git a/Engine/Source/Developer/ShaderFormatVectorVM/Private/ir_vm_merge_op_visitor.cpp b/Engine/Source/Developer/ShaderFormatVectorVM/Private/ir_vm_merge_op_visitor.cpp
--- a/Engine/Source/Developer/ShaderFormatVectorVM/Private/ir_vm_merge_op_visitor.cpp
+++ b/Engine/Source/Developer/ShaderFormatVectorVM/Private/ir_vm_merge_op_visitor.cpp
@@ -8,7 +8,7 @@
 
 PRAGMA_DISABLE_SHADOW_VARIABLE_WARNINGS
 #include "glsl_parser_extras.h"
-PRAGMA_POP
+PRAGMA_ENABLE_SHADOW_VARIABLE_WARNINGS
 
 #include "hash_table.h"
 #include "ir_rvalue_visitor.h"
diff --git a/Engine/Source/Developer/ShaderFormatVectorVM/Private/ir_vm_propagate_non_expressions_visitor.cpp b/Engine/Source/Developer/ShaderFormatVectorVM/Private/ir_vm_propagate_non_expressions_visitor.cpp
--- a/Engine/Source/Developer/ShaderFormatVectorVM/Private/ir_vm_propagate_non_expressions_visitor.cpp
+++ b/Engine/Source/Developer/ShaderFormatVectorVM/Private/ir_vm_propagate_non_expressions_visitor.cpp
@@ -9,7 +9,7 @@
 
 PRAGMA_DISABLE_SHADOW_VARIABLE_WARNINGS
 #include "glsl_parser_extras.h"
-PRAGMA_POP
+PRAGMA_ENABLE_SHADOW_VARIABLE_WARNINGS
 
 #include "hash_table.h"
 #include "ir_rvalue_visitor.h"
diff --git a/Engine/Source/Developer/ShaderFormatVectorVM/Private/ir_vm_scalarize_visitor.cpp b/Engine/Source/Developer/ShaderFormatVectorVM/Private/ir_vm_scalarize_visitor.cpp
--- a/Engine/Source/Developer/ShaderFormatVectorVM/Private/ir_vm_scalarize_visitor.cpp
+++ b/Engine/Source/Developer/ShaderFormatVectorVM/Private/ir_vm_scalarize_visitor.cpp
@@ -8,7 +8,7 @@
 
 PRAGMA_DISABLE_SHADOW_VARIABLE_WARNINGS
 #include "glsl_parser_extras.h"
-PRAGMA_POP
+PRAGMA_ENABLE_SHADOW_VARIABLE_WARNINGS
 
 #include "hash_table.h"
 #include "ir_rvalue_visitor.h"
diff --git a/Engine/Source/Developer/ShaderFormatVectorVM/Private/ir_vm_to_single_op_visitor.cpp b/Engine/Source/Developer/ShaderFormatVectorVM/Private/ir_vm_to_single_op_visitor.cpp
--- a/Engine/Source/Developer/ShaderFormatVectorVM/Private/ir_vm_to_single_op_visitor.cpp
+++ b/Engine/Source/Developer/ShaderFormatVectorVM/Private/ir_vm_to_single_op_visitor.cpp
@@ -8,7 +8,7 @@
 
 PRAGMA_DISABLE_SHADOW_VARIABLE_WARNINGS
 #include "glsl_parser_extras.h"
-PRAGMA_POP
+PRAGMA_ENABLE_SHADOW_VARIABLE_WARNINGS
 
 #include "hash_table.h"
 #include "ir_rvalue_visitor.h"
diff --git a/Engine/Source/Developer/ShaderFormatVectorVM/Private/ir_vm_vec_op_to_scalar_visitor.cpp b/Engine/Source/Developer/ShaderFormatVectorVM/Private/ir_vm_vec_op_to_scalar_visitor.cpp
--- a/Engine/Source/Developer/ShaderFormatVectorVM/Private/ir_vm_vec_op_to_scalar_visitor.cpp
+++ b/Engine/Source/Developer/ShaderFormatVectorVM/Private/ir_vm_vec_op_to_scalar_visitor.cpp
@@ -9,7 +9,7 @@
 
 PRAGMA_DISABLE_SHADOW_VARIABLE_WARNINGS
 #include "glsl_parser_extras.h"
-PRAGMA_POP
+PRAGMA_ENABLE_SHADOW_VARIABLE_WARNINGS
 
 #include "hash_table.h"
 #include "ir_rvalue_visitor.h"
diff --git a/Engine/Source/Editor/Kismet/Private/SSCSEditor.cpp b/Engine/Source/Editor/Kismet/Private/SSCSEditor.cpp
--- a/Engine/Source/Editor/Kismet/Private/SSCSEditor.cpp
+++ b/Engine/Source/Editor/Kismet/Private/SSCSEditor.cpp
@@ -1543,7 +1543,9 @@ TSharedRef<SWidget> SSCS_RowWidget::GenerateWidgetForColumn( const FName& Column
 				.Text( LOCTEXT("UnknownColumn", "Unknown Column") );
 	}
 }
+END_SLATE_FUNCTION_BUILD_OPTIMIZATION
 
+BEGIN_SLATE_FUNCTION_BUILD_OPTIMIZATION
 void SSCS_RowWidget::AddToToolTipInfoBox(const TSharedRef<SVerticalBox>& InfoBox, const FText& Key, TSharedRef<SWidget> ValueIcon, const TAttribute<FText>& Value, bool bImportant)
 {
 	InfoBox->AddSlot()
@@ -1576,7 +1578,9 @@ void SSCS_RowWidget::AddToToolTipInfoBox(const TSharedRef<SVerticalBox>& InfoBox
 			]
 		];
 }
+END_SLATE_FUNCTION_BUILD_OPTIMIZATION
 
+BEGIN_SLATE_FUNCTION_BUILD_OPTIMIZATION
 TSharedRef<SToolTip> SSCS_RowWidget::CreateToolTipWidget() const
 {
 	// Create a box to hold every line of info in the body of the tooltip
@@ -3394,6 +3398,7 @@ TSharedRef<SWidget> SSCS_RowWidget_Separator::GenerateWidgetForColumn(const FNam
 //////////////////////////////////////////////////////////////////////////
 // SSCSEditor
 
+BEGIN_SLATE_FUNCTION_BUILD_OPTIMIZATION
 void SSCSEditor::Construct( const FArguments& InArgs )
 {
 	EditorMode = InArgs._EditorMode;
diff --git a/Engine/Source/Programs/UnrealVersionSelector/Private/Linux/LinuxPlatformInstallation.cpp b/Engine/Source/Programs/UnrealVersionSelector/Private/Linux/LinuxPlatformInstallation.cpp
--- a/Engine/Source/Programs/UnrealVersionSelector/Private/Linux/LinuxPlatformInstallation.cpp
+++ b/Engine/Source/Programs/UnrealVersionSelector/Private/Linux/LinuxPlatformInstallation.cpp
@@ -218,6 +218,7 @@ public:
 			]
 		];
 	}
+	END_SLATE_FUNCTION_BUILD_OPTIMIZATION
 
 private:
 	TSharedRef<SWidget> OnGenerateWidget(TSharedPtr<FEngineInstallationInfo> Item)
@@ -357,6 +358,8 @@ public:
 			]
 		];
 	}
+	END_SLATE_FUNCTION_BUILD_OPTIMIZATION
+
 private:
 	FReply OnOkClicked()
 	{
diff --git a/Engine/Source/Runtime/AppFramework/Private/Widgets/Colors/SColorPicker.cpp b/Engine/Source/Runtime/AppFramework/Private/Widgets/Colors/SColorPicker.cpp
--- a/Engine/Source/Runtime/AppFramework/Private/Widgets/Colors/SColorPicker.cpp
+++ b/Engine/Source/Runtime/AppFramework/Private/Widgets/Colors/SColorPicker.cpp
@@ -629,6 +629,7 @@ void SColorPicker::GenerateDefaultColorPickerContent( bool bAdvancedSectionExpan
 	
 	HideSmallTrash();
 }
+END_SLATE_FUNCTION_BUILD_OPTIMIZATION
 
 EActiveTimerReturnType SColorPicker::AnimatePostConstruct( double InCurrentTime, float InDeltaTime )
 {
diff --git a/Engine/Source/Runtime/AppFramework/Private/Widgets/Testing/STestSuite.cpp b/Engine/Source/Runtime/AppFramework/Private/Widgets/Testing/STestSuite.cpp
--- a/Engine/Source/Runtime/AppFramework/Private/Widgets/Testing/STestSuite.cpp
+++ b/Engine/Source/Runtime/AppFramework/Private/Widgets/Testing/STestSuite.cpp
@@ -1903,7 +1903,6 @@ public:
 	SLATE_BEGIN_ARGS( SRichTextTest ){}
 	SLATE_END_ARGS()
 
-	BEGIN_SLATE_FUNCTION_BUILD_OPTIMIZATION
 	void Construct(const FArguments& InArgs)
 	{
 		WrapWidth = 600;
@@ -2273,7 +2272,6 @@ public:
 			]
 		];
 	}
-	END_SLATE_FUNCTION_BUILD_OPTIMIZATION
 
 	FSlateWidgetRun::FWidgetRunInfo OnCreateWidgetDecoratorWidget( const FTextRunInfo& RunInfo, const ISlateStyle* Style ) const
 	{
@@ -2493,7 +2491,6 @@ public:
 	 *
 	 * @param InArgs   Declaration from which to construct the widget
 	 */
-	BEGIN_SLATE_FUNCTION_BUILD_OPTIMIZATION
 	void Construct(const FArguments& InArgs)
 	{
 		bIsPassword = true;
@@ -2627,7 +2624,6 @@ public:
 			]
 		];
 	}
-	END_SLATE_FUNCTION_BUILD_OPTIMIZATION
 
 	void FocusDefaultWidget()
 	{
@@ -2994,7 +2990,6 @@ public:
 	 *
 	 * @param InArgs   Declaration from which to construct the widget
 	 */
-	BEGIN_SLATE_FUNCTION_BUILD_OPTIMIZATION
 	void Construct(const FArguments& InArgs)
 	{
 		// Define and add the "Roboto" font family
@@ -3370,7 +3365,6 @@ public:
 			]
 		];
 	}
-	END_SLATE_FUNCTION_BUILD_OPTIMIZATION
 
 protected:
 	
@@ -3901,7 +3895,6 @@ public:
 	SLATE_BEGIN_ARGS( SAnimTest ){}
 	SLATE_END_ARGS()
 
-	BEGIN_SLATE_FUNCTION_BUILD_OPTIMIZATION
 	void Construct( const FArguments& InArgs )
 	{
 		this->ChildSlot
@@ -4060,7 +4053,6 @@ public:
 		// Set the default Animation duration.
 		AnimTime_OnChanged(0.15f);
 	}
-	END_SLATE_FUNCTION_BUILD_OPTIMIZATION
 
 	FReply PlayAnimation_OnClicked()
 	{
@@ -4130,7 +4122,6 @@ class SFxTest : public SCompoundWidget
 	{}
 	SLATE_END_ARGS()
 
-	BEGIN_SLATE_FUNCTION_BUILD_OPTIMIZATION
 	void Construct( const FArguments& InArgs )
 	{
 		RenderScale = 1.0f;
@@ -4288,7 +4279,6 @@ class SFxTest : public SCompoundWidget
 			]
 		];
 	}
-	END_SLATE_FUNCTION_BUILD_OPTIMIZATION
 
 	bool GetFxWidgetIgnoreClipping() const { return FxWidgetIgnoreClippingState == ECheckBoxState::Checked; }
 	ECheckBoxState GetFxWidgetIgnoreClippingState() const { return FxWidgetIgnoreClippingState; }
@@ -5025,7 +5015,6 @@ public:
 	{}
 	SLATE_END_ARGS()
 
-	BEGIN_SLATE_FUNCTION_BUILD_OPTIMIZATION
 	void Construct( const FArguments& InArgs )	
 	{
 		this->ChildSlot
@@ -5144,7 +5133,6 @@ public:
 			
 		];
 	}
-	END_SLATE_FUNCTION_BUILD_OPTIMIZATION
 };
 
 class SResponsiveGridPanelTestWidget : public SUserWidget
@@ -5689,7 +5677,6 @@ TSharedRef<SRenderTransformManipulatorWidget> SRenderTransformManipulatorWidget:
 
 static TSharedPtr<FTabManager> TestSuite1TabManager;
 
-BEGIN_SLATE_FUNCTION_BUILD_OPTIMIZATION
 TSharedRef<SDockTab> SpawnTab(const FSpawnTabArgs& Args, FName TabIdentifier)
 {
 	if (TabIdentifier == FName(TEXT("AnimationTestTab")))
@@ -6015,7 +6002,6 @@ TSharedRef<SDockTab> SpawnResponsiveGrid(const FSpawnTabArgs& Args)
 
 	return ResponsiveGridTab;
 }
-END_SLATE_FUNCTION_BUILD_OPTIMIZATION
 
 
 
diff --git a/Engine/Source/Runtime/Engine/Private/SceneUtils.cpp b/Engine/Source/Runtime/Engine/Private/SceneUtils.cpp
--- a/Engine/Source/Runtime/Engine/Private/SceneUtils.cpp
+++ b/Engine/Source/Runtime/Engine/Private/SceneUtils.cpp
@@ -778,5 +778,3 @@ void FRenderQueryPool::ReleaseQuery(FRenderQueryRHIRef &Query)
 		Query = NULL;
 	}
 }
-
-PRAGMA_ENABLE_OPTIMIZATION
\ No newline at end of file
diff --git a/Engine/Source/Runtime/Landscape/Public/LandscapeDataAccess.h b/Engine/Source/Runtime/Landscape/Public/LandscapeDataAccess.h
--- a/Engine/Source/Runtime/Landscape/Public/LandscapeDataAccess.h
+++ b/Engine/Source/Runtime/Landscape/Public/LandscapeDataAccess.h
@@ -106,12 +106,6 @@ public:
 private:
 	TMap<UTexture2D*, TArray<FLockedMipDataInfo> > LockedMipInfoMap;
 };
-
[REDACTED_EMAIL]: Apparent VC10 compiler bug here causes an access violation in UnlockMip in Shipping builds
-#ifdef _MSC_VER
-PRAGMA_ENABLE_OPTIMIZATION
-#endif
-
 	
 //
 // FLandscapeComponentDataInterface
diff --git a/Engine/Source/Runtime/SlateCore/Private/Fonts/FontCacheHarfBuzz.cpp b/Engine/Source/Runtime/SlateCore/Private/Fonts/FontCacheHarfBuzz.cpp
--- a/Engine/Source/Runtime/SlateCore/Private/Fonts/FontCacheHarfBuzz.cpp
+++ b/Engine/Source/Runtime/SlateCore/Private/Fonts/FontCacheHarfBuzz.cpp
@@ -18,9 +18,6 @@
 			#define snprintf _snprintf
 		#endif
 	#endif // #if PLATFORM_WINDOWS
-	#if PLATFORM_MAC || PLATFORM_IOS || PLATFORM_TVOS
-	PRAGMA_DISABLE_DEPRECATION_WARNINGS
-	#endif
 
 	// Include some private headers needed for our font implementation
 	PRAGMA_DISABLE_DEPRECATION_WARNINGS
@@ -30,9 +27,6 @@
 	THIRD_PARTY_INCLUDES_END
 	PRAGMA_ENABLE_DEPRECATION_WARNINGS
 
-	#if PLATFORM_MAC || PLATFORM_IOS || PLATFORM_TVOS
-	PRAGMA_ENABLE_DEPRECATION_WARNINGS
-	#endif
 	#if PLATFORM_WINDOWS || PLATFORM_XBOXONE
 		#pragma warning(pop)
 	#endif // #if PLATFORM_WINDOWS
