Resolving nonunity build error in SlateViewer

diff --git a/Engine/Source/Runtime/AppFramework/Private/Framework/Testing/STableViewTesting.cpp b/Engine/Source/Runtime/AppFramework/Private/Framework/Testing/STableViewTesting.cpp
--- a/Engine/Source/Runtime/AppFramework/Private/Framework/Testing/STableViewTesting.cpp
+++ b/Engine/Source/Runtime/AppFramework/Private/Framework/Testing/STableViewTesting.cpp
@@ -29,11 +29,11 @@
 #include "Widgets/Views/STreeView.h"
 #include "Widgets/Input/SSpinBox.h"
 #include "Widgets/Input/SComboBox.h"
+#include "Widgets/Input/SNumericEntryBox.h"
+#include "Widgets/Input/SCheckBox.h"
 
 #if !UE_BUILD_SHIPPING
 
-#include "Widgets/Input/SNumericEntryBox.h"
-
 #define LOCTEXT_NAMESPACE "STableViewTesting"
 
 class FTestData;
@@ -1470,42 +1470,6 @@ protected:
 			: FText::GetEmpty();
 	}
 
-	/** 
-	+SHorizontalBox::Slot() .AutoWidth()
-				[
-					SNew(SCheckBox)
-					.IsChecked(this, &STableViewTesting::IsScrollAnimChecked)
-					.OnCheckStateChanged(this, &STableViewTesting::HandleScrollAnimCheckStateChanged)
-				]
-				+SHorizontalBox::Slot().AutoWidth() .Padding(2,0,0,0) .VAlign(VAlign_Center)
-				[
-					SNew(STextBlock).Text( LOCTEXT("AnimateScrollingLabel", "Animate Scrolling") )
-				]
-				+SHorizontalBox::Slot() .AutoWidth() .Padding(10, 0, 0, 0)
-				[
-					SNew(SCheckBox)
-					.IsChecked(this, &STableViewTesting::IsFixedItemOffsetChecked)
-					.OnCheckStateChanged(this, &STableViewTesting::HandleFixedItemOffsetCheckStateChanged)
-				]
-				+SHorizontalBox::Slot().AutoWidth()
-				[
-					SNew(SHorizontalBox)
-					.IsEnabled(this, &STableViewTesting::GetIsFixedItemOffsetEnabled)
-					+SHorizontalBox::Slot().AutoWidth()
-					[
-						SNew(STextBlock).Text(LOCTEXT("FixedLineScrollOffsetLabel", "Fixed Item Scroll Offset: "))
-					]
-					+SHorizontalBox::Slot().AutoWidth()
-					[
-						// We set how many data items we want in our test case.
-						SNew(SNumericEntryBox<float>)
-						.MinValue(0.f)
-						.MaxValue(1.f)
-						.Value(this, &STableViewTesting::GetFixedLineScrollOffsetValue)
-						.OnValueChanged(this, &STableViewTesting::HandleFixedLineScrollOffsetValueChanged)
-					]
-	*/
-
 	ECheckBoxState IsScrollAnimChecked() const
 	{
 		return bIsScrollAnimEnabled ? ECheckBoxState::Checked : ECheckBoxState::Unchecked;
