Updated the multi-line editable text example to use FCoreStyle::GetDefaultFont now that legacy fallback fonts have been disabled

This lets it continue to work with languages other than English

diff --git a/Engine/Source/Runtime/AppFramework/Private/Widgets/Testing/STestSuite.cpp b/Engine/Source/Runtime/AppFramework/Private/Widgets/Testing/STestSuite.cpp
--- a/Engine/Source/Runtime/AppFramework/Private/Widgets/Testing/STestSuite.cpp
+++ b/Engine/Source/Runtime/AppFramework/Private/Widgets/Testing/STestSuite.cpp
@@ -2783,13 +2783,10 @@ struct FTextStyles
 	/** This struct defines a font family, which combines multiple TTF or OTF fonts into a single group, allowing text to be styled as bold or italic */
 	struct FFontFamily
 	{
-		FFontFamily(FText InDisplayName, FName InFamilyName, FName InRegularFont, FName InBoldFont, FName InItalicFont, FName InBoldItalicFont)
+		FFontFamily(FText InDisplayName, FName InFamilyName, const TSharedRef<const FCompositeFont>& InCompositeFont)
 			: DisplayName(InDisplayName)
 			, FamilyName(InFamilyName)
-			, RegularFont(InRegularFont)
-			, BoldFont(InBoldFont)
-			, ItalicFont(InItalicFont)
-			, BoldItalicFont(InBoldItalicFont)
+			, CompositeFont(InCompositeFont)
 		{
 		}
 
@@ -2799,15 +2796,12 @@ struct FTextStyles
 		/** Named used to identify this family from the TextStyle decorator */
 		FName FamilyName;
 
-		/** File paths to the fonts on disk */
-		FName RegularFont;
-		FName BoldFont;
-		FName ItalicFont;
-		FName BoldItalicFont;
+		/** Composite font to use (should contain at least a Regular, Bold, Italic, and BoldItalic style) */
+		TSharedRef<const FCompositeFont> CompositeFont;
 	};
 
-	/** Convert the given text style into run meta-information, so that valid source rich-text formatting can be generated for it */
-	static FRunInfo CreateRunInfo(const TSharedPtr<FFontFamily>& InFontFamily, const uint16 InFontSize, const EFontStyle::Flags InFontStyle, const FLinearColor& InFontColor)
+	/** Get the font style name to use based on the requested style flags */
+	static FString GetFontStyleString(const EFontStyle::Flags InFontStyle)
 	{
 		FString FontStyleString;
 		if(InFontStyle == EFontStyle::Regular)
@@ -2818,18 +2812,23 @@ struct FTextStyles
 		{
 			if(InFontStyle & EFontStyle::Bold)
 			{
-				FontStyleString += "Bold";
+				FontStyleString += TEXT("Bold");
 			}
 			if(InFontStyle & EFontStyle::Italic)
 			{
-				FontStyleString += "Italic";
+				FontStyleString += TEXT("Italic");
 			}
 		}
+		return FontStyleString;
+	}
 
+	/** Convert the given text style into run meta-information, so that valid source rich-text formatting can be generated for it */
+	static FRunInfo CreateRunInfo(const TSharedPtr<FFontFamily>& InFontFamily, const uint16 InFontSize, const EFontStyle::Flags InFontStyle, const FLinearColor& InFontColor)
+	{
 		FRunInfo RunInfo(TEXT("TextStyle"));
 		RunInfo.MetaData.Add(TEXT("FontFamily"), InFontFamily->FamilyName.ToString());
 		RunInfo.MetaData.Add(TEXT("FontSize"), FString::FromInt(InFontSize));
-		RunInfo.MetaData.Add(TEXT("FontStyle"), FontStyleString);
+		RunInfo.MetaData.Add(TEXT("FontStyle"), GetFontStyleString(InFontStyle));
 		RunInfo.MetaData.Add(TEXT("FontColor"), InFontColor.ToString());
 		return RunInfo;
 	}
@@ -2885,23 +2884,13 @@ struct FTextStyles
 	/** Convert the given text style into a text block style for use by Slate */
 	static FTextBlockStyle CreateTextBlockStyle(const TSharedPtr<FFontFamily>& InFontFamily, const uint16 InFontSize, const EFontStyle::Flags InFontStyle, const FLinearColor& InFontColor)
 	{
-		FName FontName = InFontFamily->RegularFont;
-		if((InFontStyle & EFontStyle::Bold) && (InFontStyle & EFontStyle::Italic))
-		{
-			FontName = InFontFamily->BoldItalicFont;
-		}
-		else if(InFontStyle & EFontStyle::Bold)
-		{
-			FontName = InFontFamily->BoldFont;
-		}
-		else if(InFontStyle & EFontStyle::Italic)
-		{
-			FontName = InFontFamily->ItalicFont;
-		}
+		FSlateFontInfo FontInfo;
+		FontInfo.CompositeFont = InFontFamily->CompositeFont;
+		FontInfo.TypefaceFontName = *GetFontStyleString(InFontStyle);
+		FontInfo.Size = InFontSize;
 
 		FTextBlockStyle TextBlockStyle;
-		TextBlockStyle.SetFontName(FontName);
-		TextBlockStyle.SetFontSize(InFontSize);
+		TextBlockStyle.SetFont(FontInfo);
 		TextBlockStyle.SetColorAndOpacity(InFontColor);
 		return TextBlockStyle;
 	}
@@ -2995,10 +2984,7 @@ public:
 		TextStyles.AvailableFontFamilies.Emplace(MakeShareable(new FTextStyles::FFontFamily(
 			LOCTEXT("RobotoFontFamily", "Roboto"), 
 			TEXT("Roboto"), 
-			FName(*(FPaths::EngineContentDir() / TEXT("Slate/Fonts/Roboto-Regular.ttf"))),
-			FName(*(FPaths::EngineContentDir() / TEXT("Slate/Fonts/Roboto-Bold.ttf"))),
-			FName(*(FPaths::EngineContentDir() / TEXT("Slate/Fonts/Roboto-Italic.ttf"))),
-			FName(*(FPaths::EngineContentDir() / TEXT("Slate/Fonts/Roboto-BoldItalic.ttf")))
+			FCoreStyle::GetDefaultFont()
 			)));
 
 		// Set some sensible defaults (these also match the default text style of "RichText.Editor.Text"
