SWizard now has a BreadcrumbTrail. Pass in ShowBreadcrumbs(true) to enable it.

SWizard now supports more complex, branching wizards through the OnGetNextPageIndex delegate.

diff --git a/Engine/Source/Runtime/AppFramework/Private/Widgets/Workflow/SWizard.cpp b/Engine/Source/Runtime/AppFramework/Private/Widgets/Workflow/SWizard.cpp
--- a/Engine/Source/Runtime/AppFramework/Private/Widgets/Workflow/SWizard.cpp
+++ b/Engine/Source/Runtime/AppFramework/Private/Widgets/Workflow/SWizard.cpp
@@ -10,7 +10,6 @@
 #include "Widgets/Input/SCheckBox.h"
 #include "Widgets/Layout/SWidgetSwitcher.h"
 
-
 #define LOCTEXT_NAMESPACE "SWizard"
 
 
@@ -41,6 +40,14 @@ void SWizard::Construct( const FArguments& InArgs )
 	[
 		SNew(SVerticalBox)
 
+		+ SVerticalBox::Slot()
+		.AutoHeight()
+		[
+			SAssignNew(BreadcrumbTrail, SBreadcrumbTrail<int32>)
+			.OnCrumbClicked(this, &SWizard::HandleBreadcrumbClicked)
+			.Visibility(InArgs._ShowBreadcrumbs ? EVisibility::Visible : EVisibility::Collapsed)
+		]
+
 		+ SVerticalBox::Slot()
 		.FillHeight(1.0)
 		[
@@ -221,6 +228,8 @@ void SWizard::Construct( const FArguments& InArgs )
 		];
 	}
 
+	OnGetNextPageIndex = InArgs._OnGetNextPageIndex;
+
 	WidgetSwitcher->SetActiveWidgetIndex(INDEX_NONE);
 	ShowPage(InArgs._InitialPageIndex.Get());
 }
@@ -250,6 +259,7 @@ void SWizard::ShowPage( int32 PageIndex )
 	}
 	else
 	{
+		ensure(false);
 		WidgetSwitcher->SetActiveWidgetIndex(INDEX_NONE);
 	}
 }
@@ -287,24 +297,60 @@ FReply SWizard::HandleFinishButtonClicked()
 	return FReply::Handled();
 }
 
+int32 SWizard::GetNextPageIndex() const
+{
+	int32 ActivePage = WidgetSwitcher->GetActiveWidgetIndex();
+	int32 NextPage = ActivePage + 1;
+
+	if (OnGetNextPageIndex.IsBound())
+	{
+		NextPage = OnGetNextPageIndex.Execute(ActivePage);
+	}
+
+	return NextPage;
+}
+
+int32 SWizard::GetPrevPageIndex() const
+{
+	TArray<int32> PageHistory;
+	BreadcrumbTrail->GetAllCrumbData(PageHistory);
+
+	if (PageHistory.Num() < 1)
+	{
+		return INDEX_NONE;
+	}
+
+	return PageHistory[PageHistory.Num() - 1];
+}
+
+void SWizard::AdvanceToPage(int32 PageIndex)
+{
+	int32 CurrentPage = GetCurrentPageIndex();
+	if (CurrentPage != INDEX_NONE)
+	{
+		BreadcrumbTrail->PushCrumb(Pages[CurrentPage].GetName(), CurrentPage);
+	}
+	ShowPage(PageIndex);
+}
 
 FReply SWizard::HandleNextButtonClicked()
 {
-	ShowPage(WidgetSwitcher->GetActiveWidgetIndex() + 1);
+	int32 NextPage = GetNextPageIndex();
+	AdvanceToPage(NextPage);
 
 	return FReply::Handled();
 }
 
-
 bool SWizard::HandleNextButtonIsEnabled() const
 {
-	return CanShowPage(WidgetSwitcher->GetActiveWidgetIndex() + 1);
+	int32 NextPage = GetNextPageIndex();
+	return CanShowPage(NextPage);
 }
 
-
 EVisibility SWizard::HandleNextButtonVisibility() const
 {
-	if (Pages.IsValidIndex(WidgetSwitcher->GetActiveWidgetIndex() + 1))
+	int32 NextPage = GetNextPageIndex();
+	if (Pages.IsValidIndex(NextPage))
 	{
 		return EVisibility::Visible;
 	}
@@ -312,7 +358,6 @@ EVisibility SWizard::HandleNextButtonVisibility() const
 	return EVisibility::Hidden;
 }
 
-
 void SWizard::HandlePageButtonCheckStateChanged( ECheckBoxState NewState, int32 PageIndex )
 {
 	if (NewState == ECheckBoxState::Checked)
@@ -321,7 +366,6 @@ void SWizard::HandlePageButtonCheckStateChanged( ECheckBoxState NewState, int32
 	}
 }
 
-
 ECheckBoxState SWizard::HandlePageButtonIsChecked( int32 PageIndex ) const
 {
 	if (PageIndex == WidgetSwitcher->GetActiveWidgetIndex())
@@ -332,13 +376,11 @@ ECheckBoxState SWizard::HandlePageButtonIsChecked( int32 PageIndex ) const
 	return ECheckBoxState::Unchecked;
 }
 
-
 bool SWizard::HandlePageButtonIsEnabled( int32 PageIndex ) const
 {
 	return CanShowPage(PageIndex);
 }
 
-
 FReply SWizard::HandlePrevButtonClicked()
 {
 	if ( WidgetSwitcher->GetActiveWidgetIndex() == 0 && OnFirstPageBackClicked.IsBound() )
@@ -347,13 +389,14 @@ FReply SWizard::HandlePrevButtonClicked()
 	}
 	else
 	{
-		ShowPage(WidgetSwitcher->GetActiveWidgetIndex() - 1);
+		int32 PrevPage = GetPrevPageIndex();
+		ShowPage(PrevPage);
+		BreadcrumbTrail->PopCrumb();
 	}
 
 	return FReply::Handled();
 }
 
-
 bool SWizard::HandlePrevButtonIsEnabled() const
 {
 	if ( WidgetSwitcher->GetActiveWidgetIndex() == 0 && OnFirstPageBackClicked.IsBound() )
@@ -361,13 +404,14 @@ bool SWizard::HandlePrevButtonIsEnabled() const
 		return true;
 	}
 
-	return CanShowPage(WidgetSwitcher->GetActiveWidgetIndex() - 1);
+	int32 PrevPage = GetPrevPageIndex();
+	return CanShowPage(PrevPage);
 }
 
-
 EVisibility SWizard::HandlePrevButtonVisibility() const
 {
-	if (WidgetSwitcher->GetActiveWidgetIndex() > 0 || OnFirstPageBackClicked.IsBound())
+	int32 PrevPage = GetPrevPageIndex();
+	if (PrevPage != INDEX_NONE || OnFirstPageBackClicked.IsBound())
 	{
 		return EVisibility::Visible;
 	}
@@ -380,10 +424,32 @@ int32 SWizard::GetNumPages() const
 	return WidgetSwitcher->GetNumWidgets();
 }
 
-int32 SWizard::GetPageIndex(const TSharedRef<SWidget>& PageWidget) const
+int32 SWizard::GetCurrentPageIndex() const
 {
-	return WidgetSwitcher->GetWidgetIndex(PageWidget);
+	return WidgetSwitcher->GetActiveWidgetIndex();
 }
 
+int32 SWizard::GetPageIndex(const TSharedPtr<SWidget>& PageWidget) const
+{
+	if (!PageWidget.IsValid())
+	{
+		return INDEX_NONE;
+	}
+
+	return WidgetSwitcher->GetWidgetIndex(PageWidget.ToSharedRef());
+}
+
+void SWizard::HandleBreadcrumbClicked(const int32& PageIndex)
+{
+	int32 CurrentCrumb = BreadcrumbTrail->PeekCrumb();
+	while (CurrentCrumb != PageIndex)
+	{
+		BreadcrumbTrail->PopCrumb();
+		CurrentCrumb = BreadcrumbTrail->PeekCrumb();
+	}
+
+	BreadcrumbTrail->PopCrumb();
+	ShowPage(PageIndex);
+}
 
 #undef LOCTEXT_NAMESPACE
diff --git a/Engine/Source/Runtime/AppFramework/Public/Widgets/Workflow/SWizard.h b/Engine/Source/Runtime/AppFramework/Public/Widgets/Workflow/SWizard.h
--- a/Engine/Source/Runtime/AppFramework/Public/Widgets/Workflow/SWizard.h
+++ b/Engine/Source/Runtime/AppFramework/Public/Widgets/Workflow/SWizard.h
@@ -13,17 +13,18 @@
 #include "Widgets/SWidget.h"
 #include "Widgets/DeclarativeSyntaxSupport.h"
 #include "Widgets/SCompoundWidget.h"
+#include "Widgets/Navigation/SBreadcrumbTrail.h"
 #include "Styling/SlateTypes.h"
 #include "Styling/CoreStyle.h"
 #include "Framework/SlateDelegates.h"
 
 /**
- * Delegate type for getting a page's enabled state.
+ * Delegate type for getting the next page to show.
  *
- * The first parameter is the page index.
+ * The parameter is the current page index.
+ * The return value should be the next page index.
  */
-DECLARE_DELEGATE_OneParam(FOnWizardPageIsEnabled, int32)
-
+DECLARE_DELEGATE_RetVal_OneParam(int32, FOnGetPageIndex, int32)
 
 /**
  * Implements a wizard widget.
@@ -48,6 +49,9 @@ public:
 
 			/** Holds the content of the button to be shown in the wizard's page selector. */
 			SLATE_NAMED_SLOT(FArguments, ButtonContent)
+			
+			/** Holds the name of the page to be displayed in the breadcrumb trail. */
+			SLATE_ATTRIBUTE(FText, Name)
 
 			/** Holds a flag indicating whether this page can be activated (default = true). */
 			SLATE_ATTRIBUTE(bool, CanShow)
@@ -76,6 +80,7 @@ public:
 			, OnEnterDelegate(InArgs._OnEnter)
 			, OnLeaveDelegate(InArgs._OnLeave)
 			, PageContent(InArgs._PageContent)
+			, Name(InArgs._Name)
 		{ }
 
 	public:
@@ -90,6 +95,12 @@ public:
 			return Showable.Get();
 		}
 
+		/** Gets the display name of the page. */
+		FText GetName() const
+		{
+			return Name.Get();
+		}
+
 		/**
 		 * Gets the button content.
 		 *
@@ -146,6 +157,9 @@ public:
 
 		// Holds the page content.
 		TAlwaysValidWidget PageContent;
+
+		// Holds the name of the page.
+		TAttribute<FText> Name;
 	};
 
 public:
@@ -163,10 +177,11 @@ public:
 		, _DesiredSize(FVector2D(0, 0))
 		, _ShowPageList(true)
 		, _ShowCancelButton(true)
+		, _ShowBreadcrumbs(false)
 		, _PageFooter()
 	{ }
 
-		SLATE_SUPPORTS_SLOT_WITH_ARGS(FWizardPage)
+	SLATE_SUPPORTS_SLOT_WITH_ARGS(FWizardPage)
 
 		/** The button style used by next and back. */
 		SLATE_STYLE_ARGUMENT(FButtonStyle, ButtonStyle)
@@ -213,12 +228,21 @@ public:
 		/** Exposes a delegate to be invoked when the wizard's 'Previous' button is clicked. */
 		SLATE_EVENT(FOnClicked, OnPrevClicked)
 
+		/** 
+		 * Exposes a delegate to be invoked when the wizard's 'Next' button is clicked.
+		 * Passes in the current page index and expects the next page index to be returned.
+		 */
+		SLATE_EVENT(FOnGetPageIndex, OnGetNextPageIndex)
+
 		/** Holds a flag indicating whether the page list should be shown (default = true). */
 		SLATE_ARGUMENT(bool, ShowPageList)
 
 		/** Holds a flag indicating whether the cancel button should be shown (default = true). */
 		SLATE_ARGUMENT(bool, ShowCancelButton)
 
+		/** Holds a flag indicating whether a breadcrumb trail should be shown (default = false) */
+		SLATE_ARGUMENT(bool, ShowBreadcrumbs)
+
 		/** Holds an optional widget containing the contents to place above the buttons for all wizard pages */
 		SLATE_NAMED_SLOT(FArguments, PageFooter)
 
@@ -248,13 +272,19 @@ public:
 	 */
 	int32 GetNumPages() const;
 
+	/**
+	 * Get the index of the current page.
+	 * @return The index of the current page.
+	 */
+	int32 GetCurrentPageIndex() const;
+
 	/**
 	 * Gets the index of the specified wizard page widget.
 	 *
 	 * @param PageWidget The page widget to get the index for.
 	 * @return The index of the page, or INDEX_NONE if not found.
 	 */
-	int32 GetPageIndex( const TSharedRef<SWidget>& PageWidget ) const;
+	int32 GetPageIndex( const TSharedPtr<SWidget>& PageWidget ) const;
 
 	/**
 	 * Attempts to show the page with the specified index.
@@ -263,6 +293,9 @@ public:
 	 */
 	void ShowPage( int32 PageIndex );
 
+	/** Shows the page with the specified index, but preserves page history. */
+	void AdvanceToPage( int32 PageIndex );
+
 public:
 
 	// SCompoundWidget interface
@@ -318,6 +351,15 @@ private:
 	// Callback for getting the visibility of the 'Previous' button.
 	EVisibility HandlePrevButtonVisibility() const;
 
+	// Callback for when a breadcrumb is clicked.
+	void HandleBreadcrumbClicked(const int32& PageIndex);
+
+	// Get the next page index to display.
+	int32 GetNextPageIndex() const;
+
+	// Get the previous page index to display.
+	int32 GetPrevPageIndex() const;
+
 private:
 
 	// Holds the wizard's desired size.
@@ -329,6 +371,9 @@ private:
 	// Holds the widget switcher.
 	TSharedPtr<SWidgetSwitcher> WidgetSwitcher;
 
+	// Holds the breadcrumb trail.
+	TSharedPtr<SBreadcrumbTrail<int32>> BreadcrumbTrail;
+
 private:
 
 	// Holds a delegate to be invoked when the 'Cancel' button has been clicked.
@@ -345,4 +390,7 @@ private:
 
 	// Holds a delegate to be invoked when the 'Previous' button has been clicked on the first page.
 	FOnClicked OnFirstPageBackClicked;
+
+	// Exposes a delegate to be invoked when the wizard's 'Next' button is clicked to fetch the next page index.
+	FOnGetPageIndex OnGetNextPageIndex;
 };
