SWizard now supports displaying the page title (enabled with ShowPageTitle) on every page of the wizard.

diff --git a/Engine/Source/Runtime/AppFramework/Private/Widgets/Workflow/SWizard.cpp b/Engine/Source/Runtime/AppFramework/Private/Widgets/Workflow/SWizard.cpp
--- a/Engine/Source/Runtime/AppFramework/Private/Widgets/Workflow/SWizard.cpp
+++ b/Engine/Source/Runtime/AppFramework/Private/Widgets/Workflow/SWizard.cpp
@@ -47,7 +47,15 @@ void SWizard::Construct( const FArguments& InArgs )
 			.OnCrumbClicked(this, &SWizard::HandleBreadcrumbClicked)
 			.Visibility(InArgs._ShowBreadcrumbs ? EVisibility::Visible : EVisibility::Collapsed)
 		]
-
+		+ SVerticalBox::Slot()
+		.AutoHeight()
+		.Padding(0, 10, 0, 0)
+		[
+			SNew(STextBlock)
+			.TextStyle(InArgs._PageTitleTextStyle)
+			.Text(this, &SWizard::HandleGetPageTitle)
+			.Visibility(InArgs._ShowPageTitle ? EVisibility::Visible : EVisibility::Collapsed)
+		]
 		+ SVerticalBox::Slot()
 		.FillHeight(1.0)
 		[
@@ -452,4 +460,14 @@ void SWizard::HandleBreadcrumbClicked(const int32& PageIndex)
 	ShowPage(PageIndex);
 }
 
+FText SWizard::HandleGetPageTitle() const
+{
+	int32 PageIndex = GetCurrentPageIndex();
+	if (Pages.IsValidIndex(PageIndex))
+	{
+		return Pages[PageIndex].GetName();
+	}
+	return FText::GetEmpty();
+}
+
 #undef LOCTEXT_NAMESPACE
diff --git a/Engine/Source/Runtime/AppFramework/Public/Widgets/Workflow/SWizard.h b/Engine/Source/Runtime/AppFramework/Public/Widgets/Workflow/SWizard.h
--- a/Engine/Source/Runtime/AppFramework/Public/Widgets/Workflow/SWizard.h
+++ b/Engine/Source/Runtime/AppFramework/Public/Widgets/Workflow/SWizard.h
@@ -178,6 +178,10 @@ public:
 		, _ShowPageList(true)
 		, _ShowCancelButton(true)
 		, _ShowBreadcrumbs(false)
+		, _BreadcrumbButtonStyle(&FCoreStyle::Get().GetWidgetStyle< FButtonStyle >("BreadcrumbButton"))
+		, _BreadcrumbTextStyle(&FCoreStyle::Get().GetWidgetStyle< FTextBlockStyle >("NormalText"))
+		, _ShowPageTitle(false)
+		, _PageTitleTextStyle(&FCoreStyle::Get().GetWidgetStyle< FTextBlockStyle >("Wizard.PageTitle"))
 		, _PageFooter()
 	{ }
 
@@ -228,7 +232,7 @@ public:
 		/** Exposes a delegate to be invoked when the wizard's 'Previous' button is clicked. */
 		SLATE_EVENT(FOnClicked, OnPrevClicked)
 
-		/** 
+		/**
 		 * Exposes a delegate to be invoked when the wizard's 'Next' button is clicked.
 		 * Passes in the current page index and expects the next page index to be returned.
 		 */
@@ -243,6 +247,18 @@ public:
 		/** Holds a flag indicating whether a breadcrumb trail should be shown (default = false) */
 		SLATE_ARGUMENT(bool, ShowBreadcrumbs)
 
+		/** The button style to use for each breadcrumb if ShowBreadcrumbs is set to true. */
+		SLATE_STYLE_ARGUMENT(FButtonStyle, BreadcrumbButtonStyle);
+
+		/** The text style to use for breadcrumbs if ShowBreadcrumbs is set to true. */
+		SLATE_STYLE_ARGUMENT(FTextBlockStyle, BreadcrumbTextStyle);
+
+		/** Holds a flag indicating whether a page title should be shown for each page (default = false) */
+		SLATE_ARGUMENT(bool, ShowPageTitle)
+
+		/** The text style to use for showing the page title if ShowPageTitle is set to true. */
+		SLATE_STYLE_ARGUMENT(FTextBlockStyle, PageTitleTextStyle);
+
 		/** Holds an optional widget containing the contents to place above the buttons for all wizard pages */
 		SLATE_NAMED_SLOT(FArguments, PageFooter)
 
@@ -360,6 +376,9 @@ private:
 	// Get the previous page index to display.
 	int32 GetPrevPageIndex() const;
 
+	// Get the page title to display for the current page.
+	FText HandleGetPageTitle() const;
+
 private:
 
 	// Holds the wizard's desired size.
diff --git a/Engine/Source/Runtime/SlateCore/Private/Styling/CoreStyle.cpp b/Engine/Source/Runtime/SlateCore/Private/Styling/CoreStyle.cpp
--- a/Engine/Source/Runtime/SlateCore/Private/Styling/CoreStyle.cpp
+++ b/Engine/Source/Runtime/SlateCore/Private/Styling/CoreStyle.cpp
@@ -1049,6 +1049,15 @@ TSharedRef<ISlateStyle> FCoreStyle::Create( const FName& InStyleSetName )
 			);
 	}
 
+	// SWizard defaults
+	{
+		Style->Set("Wizard.PageTitle", FTextBlockStyle(NormalText)
+			.SetFont(DEFAULT_FONT("BoldCondensed", 28))
+			.SetShadowOffset(FVector2D(1, 1))
+			.SetShadowColorAndOpacity(FLinearColor(0, 0, 0, 0.9f))
+			);
+	}
+
 	// SNotificationList defaults...
 	{
 		Style->Set( "NotificationList.FontBold", DEFAULT_FONT( "Bold", 16 ) );
