Removed unused OnNextClicked and OnPrevClicked delegates from SWizard after discovering that they are, in fact, unused.

diff --git a/Engine/Source/Runtime/AppFramework/Public/Widgets/Workflow/SWizard.h b/Engine/Source/Runtime/AppFramework/Public/Widgets/Workflow/SWizard.h
--- a/Engine/Source/Runtime/AppFramework/Public/Widgets/Workflow/SWizard.h
+++ b/Engine/Source/Runtime/AppFramework/Public/Widgets/Workflow/SWizard.h
@@ -226,12 +226,6 @@ public:
 		/** Exposes a delegate to be invoked when the wizard's 'Back' button is clicked. This button is only visible if this delegate is set and only on the first page. */
 		SLATE_EVENT(FOnClicked, OnFirstPageBackClicked)
 
-		/** Exposes a delegate to be invoked when the wizard's 'Next' button is clicked. */
-		SLATE_EVENT(FOnClicked, OnNextClicked)
-
-		/** Exposes a delegate to be invoked when the wizard's 'Previous' button is clicked. */
-		SLATE_EVENT(FOnClicked, OnPrevClicked)
-
 		/**
 		 * Exposes a delegate to be invoked when the wizard's 'Next' button is clicked.
 		 * Passes in the current page index and expects the next page index to be returned.
