Removed an unnecessary 10 pixels of padding from above the title of an SWizard when a ShowBreadcrumbTrail is false.

#rnx

diff --git a/Engine/Source/Runtime/AppFramework/Private/Widgets/Workflow/SWizard.cpp b/Engine/Source/Runtime/AppFramework/Private/Widgets/Workflow/SWizard.cpp
--- a/Engine/Source/Runtime/AppFramework/Private/Widgets/Workflow/SWizard.cpp
+++ b/Engine/Source/Runtime/AppFramework/Private/Widgets/Workflow/SWizard.cpp
@@ -49,7 +49,7 @@ void SWizard::Construct( const FArguments& InArgs )
 		]
 		+ SVerticalBox::Slot()
 		.AutoHeight()
-		.Padding(0, 10, 0, 0)
+		.Padding(InArgs._ShowBreadcrumbs ? FMargin(0, 10, 0, 0) : FMargin(0))
 		[
 			SNew(STextBlock)
 			.TextStyle(InArgs._PageTitleTextStyle)
