Slate Viewer Gallery:  Added selection colors, Icon Gallery now scrolls

diff --git a/Engine/Source/Runtime/AppFramework/Private/Framework/Testing/SStarshipGallery.cpp b/Engine/Source/Runtime/AppFramework/Private/Framework/Testing/SStarshipGallery.cpp
--- a/Engine/Source/Runtime/AppFramework/Private/Framework/Testing/SStarshipGallery.cpp
+++ b/Engine/Source/Runtime/AppFramework/Private/Framework/Testing/SStarshipGallery.cpp
@@ -210,8 +210,8 @@ public:
             // = SNew(SButton)
             return SNew(SBox)
             .Padding(12)
-            .WidthOverride(128)
-            .HeightOverride(128)
+            .WidthOverride(120)
+            .HeightOverride(120)
             [
                     SNew(SOverlay)
 
@@ -280,25 +280,9 @@ public:
                     +SUniformWrapPanel::Slot()[ GenerateColorButton( FText::FromString("DROPDOWN"),        "Colors.Dropdown")]
                     +SUniformWrapPanel::Slot()[ GenerateColorButton( FText::FromString("HOVER"),           "Colors.Hover")]
                     +SUniformWrapPanel::Slot()[ GenerateColorButton( FText::FromString("HOVER2"),          "Colors.Hover2")]
-
-                ]
-
-                +SVerticalBox::Slot()
-                .AutoHeight()
-                .Padding(8.f, 48.f, 8.f, 12.f)
-                [
-                    SNew(STextBlock).ColorAndOpacity(LabelColor).Text(NSLOCTEXT("StarshipGallery", "TextIconColors", "TEXT & ICON COLORS"))
-                ]
-
-                +SVerticalBox::Slot()
-                .AutoHeight()
-                [
-                    SNew(SUniformWrapPanel)
-
                     +SUniformWrapPanel::Slot()[ GenerateColorButton( FText::FromString("WHITE"),           "Colors.White", true)]
                     +SUniformWrapPanel::Slot()[ GenerateColorButton( FText::FromString("WHITE25"),         "Colors.White25")]
-                    +SUniformWrapPanel::Slot()[ GenerateColorButton( FText::FromString("WHITE50"),         "Colors.White50")]
-                    +SUniformWrapPanel::Slot()[ GenerateColorButton( FText::FromString("HIGHLIGHT"),       "Colors.HIGHLIGHT")]
+
                 ]
 
                 +SVerticalBox::Slot()
@@ -311,6 +295,9 @@ public:
                 ]
 
 
+               
+
+
             ]
 
             +SHorizontalBox::Slot()
@@ -322,6 +309,37 @@ public:
                 +SVerticalBox::Slot()
                 .Padding(8.f, 0.f, 8.f, 12.f)
                 .AutoHeight()
+                [
+                    SNew(STextBlock).ColorAndOpacity(LabelColor).Text(NSLOCTEXT("StarshipGallery", "TextIconColors", "FOREGROUND COLORS"))
+
+                ]
+
+                +SVerticalBox::Slot()
+                .AutoHeight()
+                [
+                    SNew(SUniformWrapPanel)
+
+                    +SUniformWrapPanel::Slot()[ GenerateColorButton( FText::FromString("FOREGROUND"),        "Colors.Foreground")]
+                    +SUniformWrapPanel::Slot()[ GenerateColorButton( FText::FromString("FOREGROUND\nHOVER"), "Colors.ForegroundHover", true)]
+                    +SUniformWrapPanel::Slot()[ GenerateColorButton( FText::FromString("FOREGROUND\nINVERTED"), "Colors.ForegroundInverted")]
+                    +SUniformWrapPanel::Slot()[ GenerateColorButton( FText::FromString("HIGHLIGHT"),         "Colors.Highlight")]
+                ]
+
+
+                +SVerticalBox::Slot()
+                .AutoHeight()
+                [
+                    SNew(SUniformWrapPanel)
+                    +SUniformWrapPanel::Slot()[ GenerateColorButton( FText::FromString("SELECT"),          "Colors.Select", true)]
+                    +SUniformWrapPanel::Slot()[ GenerateColorButton( FText::FromString("SELECT\nHOVER"),   "Colors.SelectHover")]
+                    +SUniformWrapPanel::Slot()[ GenerateColorButton( FText::FromString("SELECT\nPARENT"),  "Colors.SelectParent")]
+                    +SUniformWrapPanel::Slot()[ GenerateColorButton( FText::FromString("SELECT\nINACTIVE"),"Colors.SelectInactive")]
+                ]
+
+
+                +SVerticalBox::Slot()
+                .Padding(8.f, 48.f, 8.f, 12.f)
+                .AutoHeight()
                 [
                     SNew(STextBlock).ColorAndOpacity(LabelColor).Text(NSLOCTEXT("StarshipGallery", "AccentColors", "ACCENT COLORS"))
                 ]
@@ -392,17 +410,23 @@ public:
             ];
         };
 
+
         return SNew(SBorder)
-        .Padding(48)
+        .BorderImage( FAppStyle::Get().GetBrush("ToolPanel.GroupBorder") )
         [
-            SNew(SVerticalBox)
-            +SVerticalBox::Slot().AutoHeight()[ GenerateIconLibrary(NSLOCTEXT("StarshipGallery", "GeneralIconTitle", "General"), "Content/Editor/Slate/Icons/GeneralTools")]
-            +SVerticalBox::Slot().AutoHeight()[ GenerateIconLibrary(NSLOCTEXT("StarshipGallery", "PaintIconTitle", "Paint"), "Content/Editor/Slate/Icons/Paint")]
-            +SVerticalBox::Slot().AutoHeight()[ GenerateIconLibrary(NSLOCTEXT("StarshipGallery", "LandscapeIconTitle", "Landscape"), "Content/Editor/Slate/Icons/Landscape")]
-            +SVerticalBox::Slot().AutoHeight()[ GenerateIconLibrary(NSLOCTEXT("StarshipGallery", "ModelingIconTitle", "Modeling"), "/Plugins/Experimental/ModelingToolsEditorMode/Content/Icons")]
-            +SVerticalBox::Slot().AutoHeight()[ GenerateIconLibrary(NSLOCTEXT("StarshipGallery", "FractureIconTitle", "Fracture"), "/Plugins/Experimental/ChaosEditor/Content")]
-            +SVerticalBox::Slot().AutoHeight()[ GenerateIconLibrary(NSLOCTEXT("StarshipGallery", "CurveEditorIconTitle", "CurveEditor"), "Content/Editor/Slate/GenericCurveEditor/Icons")]
-            +SVerticalBox::Slot().AutoHeight()[ GenerateIconLibrary(NSLOCTEXT("StarshipGallery", "QuixelIconTitle", "Quixel"), "Content/Slate/Starship/Icons/Quixel")]
+            SNew(SScrollBox)
+            + SScrollBox::Slot()
+            .Padding(48)
+            [
+                SNew(SVerticalBox)
+                +SVerticalBox::Slot().AutoHeight()[ GenerateIconLibrary(NSLOCTEXT("StarshipGallery", "GeneralIconTitle", "General"), "Content/Editor/Slate/Icons/GeneralTools")]
+                +SVerticalBox::Slot().AutoHeight()[ GenerateIconLibrary(NSLOCTEXT("StarshipGallery", "PaintIconTitle", "Paint"), "Content/Editor/Slate/Icons/Paint")]
+                +SVerticalBox::Slot().AutoHeight()[ GenerateIconLibrary(NSLOCTEXT("StarshipGallery", "LandscapeIconTitle", "Landscape"), "Content/Editor/Slate/Icons/Landscape")]
+                +SVerticalBox::Slot().AutoHeight()[ GenerateIconLibrary(NSLOCTEXT("StarshipGallery", "ModelingIconTitle", "Modeling"), "/Plugins/Experimental/ModelingToolsEditorMode/Content/Icons")]
+                +SVerticalBox::Slot().AutoHeight()[ GenerateIconLibrary(NSLOCTEXT("StarshipGallery", "FractureIconTitle", "Fracture"), "/Plugins/Experimental/ChaosEditor/Content")]
+                +SVerticalBox::Slot().AutoHeight()[ GenerateIconLibrary(NSLOCTEXT("StarshipGallery", "CurveEditorIconTitle", "CurveEditor"), "Content/Editor/Slate/GenericCurveEditor/Icons")]
+                +SVerticalBox::Slot().AutoHeight()[ GenerateIconLibrary(NSLOCTEXT("StarshipGallery", "QuixelIconTitle", "Quixel"), "Content/Slate/Starship/Icons/Quixel")]
+            ]
         ];
     }
 
@@ -667,7 +691,7 @@ public:
                 SNew(SButton)
                 .ButtonStyle( &FAppStyle::Get().GetWidgetStyle< FButtonStyle >( "PrimaryButton" ) )
                 [
-                    LeftRightLabel("PerspectiveBox", LOCTEXT("PrimaryButtonExampleLabel", "PRIMARY BUTTON"))
+                    LeftRightLabel("Icons.box-perspective", LOCTEXT("PrimaryButtonExampleLabel", "PRIMARY BUTTON"))
                 ]
             ]
 
@@ -679,7 +703,7 @@ public:
                 .ButtonStyle( &FAppStyle::Get().GetWidgetStyle< FButtonStyle >( "PrimaryButton" ) )
                 .VAlign(VAlign_Center)
                 [
-                    LeftRightLabel("PerspectiveBox")
+                    LeftRightLabel("Icons.box-perspective")
                 ]
             ]
         ];
@@ -708,7 +732,7 @@ public:
                 SNew(SButton)
                 .ButtonStyle( &FAppStyle::Get().GetWidgetStyle< FButtonStyle >( "Button" ) )
                 [
-                    LeftRightLabel(("PerspectiveBox"), LOCTEXT("ButtonExampleLabel", "NORMAL BUTTON"))
+                    LeftRightLabel("Icons.box-perspective", LOCTEXT("ButtonExampleLabel", "NORMAL BUTTON"))
                 ]
             ]
 
@@ -720,7 +744,7 @@ public:
                 .ButtonStyle( &FAppStyle::Get().GetWidgetStyle< FButtonStyle >( "Button" ) )
                 .VAlign(VAlign_Center)
                 [
-                    LeftRightLabel("PerspectiveBox")
+                    LeftRightLabel("Icons.box-perspective")
                 ]
             ]
 
@@ -752,7 +776,7 @@ public:
                 SNew(SButton)
                 .ButtonStyle( &FAppStyle::Get().GetWidgetStyle< FButtonStyle >( "SimpleButton" ) )
                 [
-                    LeftRightLabel("PerspectiveBox", LOCTEXT("TextButtonExampleLabel", "SIMPLE BUTTON")) 
+                    LeftRightLabel("Icons.box-perspective", LOCTEXT("TextButtonExampleLabel", "SIMPLE BUTTON")) 
                 ]
             ]
 
@@ -763,7 +787,7 @@ public:
                 SNew(SButton)
                 .ButtonStyle( &FAppStyle::Get().GetWidgetStyle< FButtonStyle >( "SimpleButton" ) )
                 [
-                    LeftRightLabel("PerspectiveBox")
+                    LeftRightLabel("Icons.box-perspective")
                 ]
             ]
 
@@ -783,7 +807,7 @@ public:
                 SNew(SCheckBox)
                 .Style( &FAppStyle::Get().GetWidgetStyle<FCheckBoxStyle>("ToggleButtonCheckBoxAlt"))
                 [
-                    LeftRightLabel("PerspectiveBox")
+                    LeftRightLabel("Icons.box-perspective")
                 ]
             ]
 
@@ -795,7 +819,7 @@ public:
                 .Style( &FAppStyle::Get().GetWidgetStyle<FCheckBoxStyle>("ToggleButtonCheckBoxAlt"))
                 .IsChecked(true)
                 [
-                    LeftRightLabel("Pyramid")
+                    LeftRightLabel("Icons.pyramid")
                 ]
             ]
 
@@ -806,7 +830,7 @@ public:
                 SNew(SCheckBox)
                 .Style( &FAppStyle::Get().GetWidgetStyle<FCheckBoxStyle>("ToggleButtonCheckBoxAlt"))
                 [
-                    LeftRightLabel("Cylinder")
+                    LeftRightLabel("Icons.cylinder")
                 ]
             ]
 
@@ -818,7 +842,7 @@ public:
                 .Style( &FAppStyle::Get().GetWidgetStyle<FCheckBoxStyle>("ToggleButtonCheckBoxAlt"))
                 .IsChecked(true)
                 [
-                    LeftRightLabel("Sphere")
+                    LeftRightLabel("Icons.sphere")
                 ]
             ]
 
@@ -837,7 +861,7 @@ public:
                 .Style( &FAppStyle::Get().GetWidgetStyle<FCheckBoxStyle>("ToggleButtonCheckBox"))
                 .IsChecked(true)
                 [
-                    LeftRightLabel("PerspectiveBox", LOCTEXT("Box", "BOX"))
+                    LeftRightLabel("Icons.box-perspective", LOCTEXT("Box", "BOX"))
                 ]
             ]
 
@@ -849,7 +873,7 @@ public:
                 .Style( &FAppStyle::Get().GetWidgetStyle<FCheckBoxStyle>("ToggleButtonCheckBox"))
                 .IsChecked(true)
                 [
-                    LeftRightLabel("Pyramid", LOCTEXT("Pyramid", "PYRAMID"))
+                    LeftRightLabel("Icons.pyramid", LOCTEXT("Pyramid", "PYRAMID"))
                 ]
             ]
 
@@ -860,7 +884,7 @@ public:
                 SNew(SCheckBox)
                 .Style( &FAppStyle::Get().GetWidgetStyle<FCheckBoxStyle>("ToggleButtonCheckBox"))
                 [
-                    LeftRightLabel("Cylinder", LOCTEXT("Cylinder", "CYLINDER"))
+                    LeftRightLabel("Icons.cylinder", LOCTEXT("Cylinder", "CYLINDER"))
                 ]
             ]
 
@@ -873,7 +897,7 @@ public:
                 SNew(SCheckBox)
                 .Style( &FAppStyle::Get().GetWidgetStyle<FCheckBoxStyle>("ToggleButtonCheckBox"))
                 [
-                    LeftRightLabel("Sphere", LOCTEXT("Sphere", "SPHERE"))
+                    LeftRightLabel("Icons.sphere", LOCTEXT("Sphere", "SPHERE"))
                 ]
             ]
 
@@ -892,7 +916,7 @@ public:
                 .IsChecked_Lambda( [this] () -> ECheckBoxState { return SegmentedBoxChoice == 0 ? ECheckBoxState::Checked : ECheckBoxState::Unchecked; })
                 .OnCheckStateChanged_Lambda( [this] (ECheckBoxState InState) { if (InState == ECheckBoxState::Checked ) SegmentedBoxChoice = 0; } )
                 [
-                    LeftRightLabel("PerspectiveBox")
+                    LeftRightLabel("Icons.box-perspective")
                 ]
             ]
 
@@ -904,7 +928,7 @@ public:
                 .IsChecked_Lambda( [this] () -> ECheckBoxState { return SegmentedBoxChoice == 1 ? ECheckBoxState::Checked : ECheckBoxState::Unchecked; })
                 .OnCheckStateChanged_Lambda( [this] (ECheckBoxState InState) { if (InState == ECheckBoxState::Checked ) SegmentedBoxChoice = 1; } )
                 [
-                    LeftRightLabel("Cylinder")
+                    LeftRightLabel("Icons.cylinder")
                 ]
             ]
 
@@ -916,7 +940,7 @@ public:
                 .IsChecked_Lambda( [this] () -> ECheckBoxState { return SegmentedBoxChoice == 2 ? ECheckBoxState::Checked : ECheckBoxState::Unchecked; })
                 .OnCheckStateChanged_Lambda( [this] (ECheckBoxState InState) { if (InState == ECheckBoxState::Checked ) SegmentedBoxChoice = 2; } )
                 [
-                    LeftRightLabel("Pyramid")
+                    LeftRightLabel("Icons.pyramid")
                 ]
             ]
 
@@ -928,7 +952,7 @@ public:
                 .IsChecked_Lambda( [this] () -> ECheckBoxState { return SegmentedBoxChoice == 3 ? ECheckBoxState::Checked : ECheckBoxState::Unchecked; })
                 .OnCheckStateChanged_Lambda( [this] (ECheckBoxState InState) { if (InState == ECheckBoxState::Checked ) SegmentedBoxChoice = 3; } )
                 [
-                    LeftRightLabel("Sphere")
+                    LeftRightLabel("Icons.sphere")
                 ]
             ]
         ];
@@ -946,7 +970,7 @@ public:
                 .IsChecked_Lambda( [this] () -> ECheckBoxState { return SegmentedBoxChoice == 0 ? ECheckBoxState::Checked : ECheckBoxState::Unchecked; })
                 .OnCheckStateChanged_Lambda( [this] (ECheckBoxState InState) { if (InState == ECheckBoxState::Checked ) SegmentedBoxChoice = 0; } )
                 [
-                    LeftRightLabel("PerspectiveBox", LOCTEXT("Box", "BOX"))
+                    LeftRightLabel("Icons.box-perspective", LOCTEXT("Box", "BOX"))
                 ]
             ]
 
@@ -958,7 +982,7 @@ public:
                 .IsChecked_Lambda( [this] () -> ECheckBoxState { return SegmentedBoxChoice == 1 ? ECheckBoxState::Checked : ECheckBoxState::Unchecked; })
                 .OnCheckStateChanged_Lambda( [this] (ECheckBoxState InState) { if (InState == ECheckBoxState::Checked ) SegmentedBoxChoice = 1; } )
                 [
-                    LeftRightLabel("Cylinder", LOCTEXT("Cylinder", "CYLINDER"))
+                    LeftRightLabel("Icons.cylinder", LOCTEXT("Cylinder", "CYLINDER"))
                 ]
             ]
 
@@ -970,7 +994,7 @@ public:
                 .IsChecked_Lambda( [this] () -> ECheckBoxState { return SegmentedBoxChoice == 2 ? ECheckBoxState::Checked : ECheckBoxState::Unchecked; })
                 .OnCheckStateChanged_Lambda( [this] (ECheckBoxState InState) { if (InState == ECheckBoxState::Checked ) SegmentedBoxChoice = 2; } )
                 [
-                    LeftRightLabel("Pyramid", LOCTEXT("Pyramid", "PYRAMID"))
+                    LeftRightLabel("Icons.pyramid", LOCTEXT("Pyramid", "PYRAMID"))
                 ]
             ]
 
@@ -982,7 +1006,7 @@ public:
                 .IsChecked_Lambda( [this] () -> ECheckBoxState { return SegmentedBoxChoice == 3 ? ECheckBoxState::Checked : ECheckBoxState::Unchecked; })
                 .OnCheckStateChanged_Lambda( [this] (ECheckBoxState InState) { if (InState == ECheckBoxState::Checked ) SegmentedBoxChoice = 3; } )
                 [
-                    LeftRightLabel("Sphere", LOCTEXT("Sphere", "SPHERE"))
+                    LeftRightLabel("Icons.sphere", LOCTEXT("Sphere", "SPHERE"))
                 ]
             ]
         ];
@@ -1058,6 +1082,7 @@ public:
         .FillWidth(1.0)
         [
             SNew(SComboButton)
+            .ComboButtonStyle( &FAppStyle::Get().GetWidgetStyle< FComboButtonStyle >("SimpleComboButton"))
             .ButtonContent()
             [
                 SNew(STextBlock)
@@ -1088,7 +1113,7 @@ public:
         NextSlot(WidgetGrid, LOCTEXT("SComboBoxIconLabel", "SimpleComboBox"))
         [
             SNew(SComboBox<TSharedPtr<FString> >)
-            .ComboBoxStyle( &FAppStyle::Get().GetWidgetStyle< FComboBoxStyle >("SimpleComboBox"))
+            .ComboBoxStyle( FAppStyle::Get(), "SimpleComboBox")
             .OptionsSource(&ComboItems)
             .OnGenerateWidget_Lambda([](TSharedPtr<FString> Item) 
             { 
@@ -1097,7 +1122,7 @@ public:
             [
                 SNew(SImage)
                 .ColorAndOpacity(FSlateColor::UseForeground())
-                .Image(FAppStyle::Get().GetBrush("PerspectiveBox"))
+                .Image(FAppStyle::Get().GetBrush("Icons.box-perspective"))
             ]
 
         ];
@@ -1593,28 +1618,24 @@ public:
             SNew(SSearchBox)
         ];
 
- 
 
-
-        /*
+        // Window Flash Button
         NextSlot(WidgetGrid, LOCTEXT("WindowFlashAction", "Window Flash Button"))
         [
-
             SNew(SButton)
             .Text(LOCTEXT("WindowFlashLabel", "Window Flash"))
             .OnClicked_Lambda( [] () -> FReply { FSlateApplication::Get().GetActiveTopLevelWindow()->FlashWindow(); return FReply::Handled(); } )
 
-        ];*/
+        ];
 
 
         return SNew(SBorder)
-        .Padding(48.f)
         .BorderImage( FAppStyle::Get().GetBrush("ToolPanel.GroupBorder") )
         [
             SNew(SScrollBox)
 
             + SScrollBox::Slot()
-            .Padding(5.0f)
+            .Padding(48.0f)
             [
                 WidgetGrid.ToSharedRef()
             ]
