Widget Gallery updated with ForegroundHeader color

diff --git a/Engine/Source/Runtime/AppFramework/Private/Framework/Testing/SStarshipGallery.cpp b/Engine/Source/Runtime/AppFramework/Private/Framework/Testing/SStarshipGallery.cpp
--- a/Engine/Source/Runtime/AppFramework/Private/Framework/Testing/SStarshipGallery.cpp
+++ b/Engine/Source/Runtime/AppFramework/Private/Framework/Testing/SStarshipGallery.cpp
@@ -253,114 +253,124 @@ public:
         FSlateColor LabelColor = FAppStyle::Get().GetSlateColor("Colors.White50");
 
         return SNew(SBorder)
+        [
+
+         SNew(SScrollBox)
+        + SScrollBox::Slot()
         .Padding(48)
         [
-            SNew(SHorizontalBox)
+        
 
-            +SHorizontalBox::Slot()
-            .FillWidth(1.0)
-            [
-                SNew(SVerticalBox)
+            SNew(SVerticalBox)
 
-                +SVerticalBox::Slot()
-                .AutoHeight()
-                .Padding(8.f, 0.f, 8.f, 12.f)
-                [
-                    SNew(STextBlock).ColorAndOpacity(LabelColor).Text(NSLOCTEXT("StarshipGallery", "BaseColors", "BASE COLORS"))
-                ]
+            +SVerticalBox::Slot()
+            [
+                SNew(SHorizontalBox)
 
-                +SVerticalBox::Slot()
-                .AutoHeight()
+                +SHorizontalBox::Slot()
+                .FillWidth(1.0)
                 [
-                    SNew(SUniformWrapPanel)
-                    +SUniformWrapPanel::Slot()[ GenerateColorButton( FText::FromString("BLACK"),           "Colors.Black")]
-                    +SUniformWrapPanel::Slot()[ GenerateColorButton( FText::FromString("FOLDOUT"),         "Colors.Foldout")]
-                    +SUniformWrapPanel::Slot()[ GenerateColorButton( FText::FromString("INPUT"),           "Colors.Input")]
-                    +SUniformWrapPanel::Slot()[ GenerateColorButton( FText::FromString("BACKGROUND"),      "Colors.Background", false, true)]
-                    +SUniformWrapPanel::Slot()[ GenerateColorButton( FText::FromString("HEADER"),          "Colors.Header")]
-                    +SUniformWrapPanel::Slot()[ GenerateColorButton( FText::FromString("DROPDOWN"),        "Colors.Dropdown")]
-                    +SUniformWrapPanel::Slot()[ GenerateColorButton( FText::FromString("HOVER"),           "Colors.Hover")]
-                    +SUniformWrapPanel::Slot()[ GenerateColorButton( FText::FromString("HOVER2"),          "Colors.Hover2")]
-                    +SUniformWrapPanel::Slot()[ GenerateColorButton( FText::FromString("WHITE"),           "Colors.White", true)]
-                    +SUniformWrapPanel::Slot()[ GenerateColorButton( FText::FromString("WHITE25"),         "Colors.White25")]
-
-                ]
+                    SNew(SVerticalBox)
 
-                +SVerticalBox::Slot()
-                .AutoHeight()
-                [
-                    SNew(SUniformWrapPanel)
-                    +SUniformWrapPanel::Slot()[ GenerateColorButton( FText::FromString("PRIMARY"),         "Colors.Primary", true)]
-                    +SUniformWrapPanel::Slot()[ GenerateColorButton( FText::FromString("PRIMARY\nHOVER"),  "Colors.PrimaryHover", true)]
-                    +SUniformWrapPanel::Slot()[ GenerateColorButton( FText::FromString("PRIMARY\nPRESS"),  "Colors.PrimaryPress", true)]
-                ]
-            ]
+                    +SVerticalBox::Slot()
+                    .AutoHeight()
+                    .Padding(8.f, 0.f, 8.f, 12.f)
+                    [
+                        SNew(STextBlock).ColorAndOpacity(LabelColor).Text(NSLOCTEXT("StarshipGallery", "BaseColors", "BASE COLORS"))
+                    ]
 
-            +SHorizontalBox::Slot()
-            .FillWidth(1.0)
-            .Padding(48.0, 0.0)
-            [
-                SNew(SVerticalBox)
+                    +SVerticalBox::Slot()
+                    .AutoHeight()
+                    [
+                        SNew(SUniformWrapPanel)
+                        +SUniformWrapPanel::Slot()[ GenerateColorButton( FText::FromString("BLACK"),           "Colors.Black")]
+                        +SUniformWrapPanel::Slot()[ GenerateColorButton( FText::FromString("FOLDOUT"),         "Colors.Foldout")]
+                        +SUniformWrapPanel::Slot()[ GenerateColorButton( FText::FromString("INPUT"),           "Colors.Input")]
+                        +SUniformWrapPanel::Slot()[ GenerateColorButton( FText::FromString("BACKGROUND"),      "Colors.Background", false, true)]
+                        +SUniformWrapPanel::Slot()[ GenerateColorButton( FText::FromString("HEADER"),          "Colors.Header")]
+                        +SUniformWrapPanel::Slot()[ GenerateColorButton( FText::FromString("DROPDOWN"),        "Colors.Dropdown")]
+                        +SUniformWrapPanel::Slot()[ GenerateColorButton( FText::FromString("HOVER"),           "Colors.Hover")]
+                        +SUniformWrapPanel::Slot()[ GenerateColorButton( FText::FromString("HOVER2"),          "Colors.Hover2")]
+                        +SUniformWrapPanel::Slot()[ GenerateColorButton( FText::FromString("WHITE"),           "Colors.White", true)]
+                        +SUniformWrapPanel::Slot()[ GenerateColorButton( FText::FromString("WHITE25"),         "Colors.White25")]
 
-                +SVerticalBox::Slot()
-                .Padding(8.f, 0.f, 8.f, 12.f)
-                .AutoHeight()
-                [
-                    SNew(STextBlock).ColorAndOpacity(LabelColor).Text(NSLOCTEXT("StarshipGallery", "TextIconColors", "FOREGROUND COLORS"))
+                    ]
 
+                    +SVerticalBox::Slot()
+                    .AutoHeight()
+                    [
+                        SNew(SUniformWrapPanel)
+                        +SUniformWrapPanel::Slot()[ GenerateColorButton( FText::FromString("PRIMARY"),         "Colors.Primary", true)]
+                        +SUniformWrapPanel::Slot()[ GenerateColorButton( FText::FromString("PRIMARY\nHOVER"),  "Colors.PrimaryHover", true)]
+                        +SUniformWrapPanel::Slot()[ GenerateColorButton( FText::FromString("PRIMARY\nPRESS"),  "Colors.PrimaryPress", true)]
+                    ]
                 ]
 
-                +SVerticalBox::Slot()
-                .AutoHeight()
+                +SHorizontalBox::Slot()
+                .FillWidth(1.0)
+                .Padding(48.0, 0.0)
                 [
-                    SNew(SUniformWrapPanel)
-
-                    +SUniformWrapPanel::Slot()[ GenerateColorButton( FText::FromString("FOREGROUND"),        "Colors.Foreground")]
-                    +SUniformWrapPanel::Slot()[ GenerateColorButton( FText::FromString("FOREGROUND\nHOVER"), "Colors.ForegroundHover", true)]
-                    +SUniformWrapPanel::Slot()[ GenerateColorButton( FText::FromString("FOREGROUND\nINVERTED"), "Colors.ForegroundInverted")]
-                    +SUniformWrapPanel::Slot()[ GenerateColorButton( FText::FromString("HIGHLIGHT"),         "Colors.Highlight")]
-                ]
+                    SNew(SVerticalBox)
 
+                    +SVerticalBox::Slot()
+                    .Padding(8.f, 0.f, 8.f, 12.f)
+                    .AutoHeight()
+                    [
+                        SNew(STextBlock).ColorAndOpacity(LabelColor).Text(NSLOCTEXT("StarshipGallery", "TextIconColors", "FOREGROUND COLORS"))
 
-                +SVerticalBox::Slot()
-                .AutoHeight()
-                [
-                    SNew(SUniformWrapPanel)
-                    +SUniformWrapPanel::Slot()[ GenerateColorButton( FText::FromString("SELECT"),          "Colors.Select", true)]
-                    +SUniformWrapPanel::Slot()[ GenerateColorButton( FText::FromString("SELECT\nHOVER"),   "Colors.SelectHover")]
-                    +SUniformWrapPanel::Slot()[ GenerateColorButton( FText::FromString("SELECT\nPARENT"),  "Colors.SelectParent")]
-                    +SUniformWrapPanel::Slot()[ GenerateColorButton( FText::FromString("SELECT\nINACTIVE"),"Colors.SelectInactive")]
-                ]
+                    ]
 
+                    +SVerticalBox::Slot()
+                    .AutoHeight()
+                    [
+                        SNew(SUniformWrapPanel)
 
-                +SVerticalBox::Slot()
-                .Padding(8.f, 48.f, 8.f, 12.f)
-                .AutoHeight()
-                [
-                    SNew(STextBlock).ColorAndOpacity(LabelColor).Text(NSLOCTEXT("StarshipGallery", "AccentColors", "ACCENT COLORS"))
-                ]
+                        +SUniformWrapPanel::Slot()[ GenerateColorButton( FText::FromString("FOREGROUND"),        "Colors.Foreground")]
+                        +SUniformWrapPanel::Slot()[ GenerateColorButton( FText::FromString("FOREGROUND\nHOVER"), "Colors.ForegroundHover", true)]
+                        +SUniformWrapPanel::Slot()[ GenerateColorButton( FText::FromString("FOREGROUND\nINVERTED"), "Colors.ForegroundInverted")]
+                        +SUniformWrapPanel::Slot()[ GenerateColorButton( FText::FromString("HIGHLIGHT"),         "Colors.Highlight")]
+                        +SUniformWrapPanel::Slot()[ GenerateColorButton( FText::FromString("FOREGROUND\nHEADER"), "Colors.ForegroundHeader", true)]
+                    ]
 
-                +SVerticalBox::Slot()
-                .AutoHeight()
-                [
-                    SNew(SUniformWrapPanel)
-                    +SUniformWrapPanel::Slot()[ GenerateColorButton( FText::FromString("BLUE"),      "Colors.AccentBlue", true)]
-                    +SUniformWrapPanel::Slot()[ GenerateColorButton( FText::FromString("PURPLE"),    "Colors.AccentPurple")]
-                    +SUniformWrapPanel::Slot()[ GenerateColorButton( FText::FromString("PINK"),      "Colors.AccentPink", true)]
-                    +SUniformWrapPanel::Slot()[ GenerateColorButton( FText::FromString("RED"),       "Colors.AccentRed", true)]
-                    +SUniformWrapPanel::Slot()[ GenerateColorButton( FText::FromString("ORANGE"),    "Colors.AccentOrange", true)]
-                    +SUniformWrapPanel::Slot()[ GenerateColorButton( FText::FromString("YELLOW"),    "Colors.AccentYellow", true)]
-                    +SUniformWrapPanel::Slot()[ GenerateColorButton( FText::FromString("GREEN"),     "Colors.AccentGreen", true)]
-                    +SUniformWrapPanel::Slot()[ GenerateColorButton( FText::FromString("BROWN"),     "Colors.AccentBrown")]
-                    +SUniformWrapPanel::Slot()[ GenerateColorButton( FText::FromString("FOLDER"),    "Colors.AccentFolder")]
-                    +SUniformWrapPanel::Slot()[ GenerateColorButton( FText::FromString("BLACK"),     "Colors.AccentBlack", false, true)]
-                    +SUniformWrapPanel::Slot()[ GenerateColorButton( FText::FromString("GRAY"),      "Colors.AccentGray")]
-                    +SUniformWrapPanel::Slot()[ GenerateColorButton( FText::FromString("WHITE"),     "Colors.AccentWhite", true)]
+                    +SVerticalBox::Slot()
+                    .AutoHeight()
+                    [
+                        SNew(SUniformWrapPanel)
+                        +SUniformWrapPanel::Slot()[ GenerateColorButton( FText::FromString("SELECT"),          "Colors.Select", true)]
+                        +SUniformWrapPanel::Slot()[ GenerateColorButton( FText::FromString("SELECT\nHOVER"),   "Colors.SelectHover")]
+                        +SUniformWrapPanel::Slot()[ GenerateColorButton( FText::FromString("SELECT\nPARENT"),  "Colors.SelectParent")]
+                        +SUniformWrapPanel::Slot()[ GenerateColorButton( FText::FromString("SELECT\nINACTIVE"),"Colors.SelectInactive")]
+                    ]
                 ]
+            ]
 
+            +SVerticalBox::Slot()
+            .Padding(48.f)
+            .AutoHeight()
+            [
+                SNew(STextBlock).ColorAndOpacity(LabelColor).Text(NSLOCTEXT("StarshipGallery", "AccentColors", "ACCENT COLORS"))
+            ]
 
+            +SVerticalBox::Slot()
+            .AutoHeight()
+            [
+                SNew(SUniformWrapPanel)
+                +SUniformWrapPanel::Slot()[ GenerateColorButton( FText::FromString("BLUE"),      "Colors.AccentBlue", true)]
+                +SUniformWrapPanel::Slot()[ GenerateColorButton( FText::FromString("PURPLE"),    "Colors.AccentPurple")]
+                +SUniformWrapPanel::Slot()[ GenerateColorButton( FText::FromString("PINK"),      "Colors.AccentPink", true)]
+                +SUniformWrapPanel::Slot()[ GenerateColorButton( FText::FromString("RED"),       "Colors.AccentRed", true)]
+                +SUniformWrapPanel::Slot()[ GenerateColorButton( FText::FromString("ORANGE"),    "Colors.AccentOrange", true)]
+                +SUniformWrapPanel::Slot()[ GenerateColorButton( FText::FromString("YELLOW"),    "Colors.AccentYellow", true)]
+                +SUniformWrapPanel::Slot()[ GenerateColorButton( FText::FromString("GREEN"),     "Colors.AccentGreen", true)]
+                +SUniformWrapPanel::Slot()[ GenerateColorButton( FText::FromString("BROWN"),     "Colors.AccentBrown")]
+                +SUniformWrapPanel::Slot()[ GenerateColorButton( FText::FromString("FOLDER"),    "Colors.AccentFolder")]
+                +SUniformWrapPanel::Slot()[ GenerateColorButton( FText::FromString("BLACK"),     "Colors.AccentBlack", false, true)]
+                +SUniformWrapPanel::Slot()[ GenerateColorButton( FText::FromString("GRAY"),      "Colors.AccentGray")]
+                +SUniformWrapPanel::Slot()[ GenerateColorButton( FText::FromString("WHITE"),     "Colors.AccentWhite", true)]
             ]
-        ];
+        ]];
+
+
     };
 
     TArray< TUniquePtr< FSlateBrush > > DynamicBrushes;
@@ -481,20 +491,20 @@ public:
 
             +SHorizontalBox::Slot()
             .AutoWidth()
-            .Padding(48.f, 0.0f)
+            .Padding(64.f, 0.0f)
             [
 
                 SNew(SVerticalBox)
 
                 // Normal 
-                +SVerticalBox::Slot().Padding(0.f, 32.f, 0.f, 4.f).AutoHeight() [ SNew(STextBlock).ColorAndOpacity(LabelColor).Text( NSLOCTEXT("StarshipGallery", "NormalBodyTextDesc", "ROBOTO 10 SLATE    [ ROBOTO 13 FIGMA ]") ) ]
+                +SVerticalBox::Slot().Padding(0.f, 64.f, 0.f, 4.f).AutoHeight() [ SNew(STextBlock).ColorAndOpacity(LabelColor).Text( NSLOCTEXT("StarshipGallery", "NormalBodyTextDesc", "ROBOTO 10 SLATE    [ ROBOTO 13 FIGMA ]") ) ]
                 +SVerticalBox::Slot().Padding(0.f, 4.f).AutoHeight() [ 
                     SNew(STextBlock)
                     .Text(NSLOCTEXT("StarshipGallery", "NormalBodyText", "Normal Text\n\nThe quick brown fox jumps over the lazy dog.\n\nNORMAL TEXT CAPITAL") ) 
                 ]
 
                 // Small 
-                +SVerticalBox::Slot().Padding(0.f, 32.f, 0.f, 4.f).AutoHeight()[ SNew(STextBlock).ColorAndOpacity(LabelColor).Text(NSLOCTEXT("StarshipGallery", "SmallBodyTextDesc", "ROBOTO 8 SLATE   [ ROBOTO 11 FIGMA ]" ) ) ]
+                +SVerticalBox::Slot().Padding(0.f, 64.f, 0.f, 4.f).AutoHeight()[ SNew(STextBlock).ColorAndOpacity(LabelColor).Text(NSLOCTEXT("StarshipGallery", "SmallBodyTextDesc", "ROBOTO 8 SLATE   [ ROBOTO 11 FIGMA ]" ) ) ]
                 +SVerticalBox::Slot().Padding(0.f, 4.f).AutoHeight()[ 
                     SNew(STextBlock)
                     .Font(FAppStyle::Get().GetFontStyle("SmallFont"))
@@ -523,12 +533,12 @@ public:
 
             +SHorizontalBox::Slot()
             .AutoWidth()
-            .Padding(48.f, 0.0f)
+            .Padding(64.f, 0.0f)
             [
                 SNew(SVerticalBox)
 
                 // Normal Bold
-                +SVerticalBox::Slot().Padding(0.f, 32.f, 0.f, 4.f).AutoHeight() [ SNew(STextBlock).ColorAndOpacity(LabelColor).Text( NSLOCTEXT("StarshipGallery", "NormalBoldTextDesc", "ROBOTO BOLD 10 SLATE   [ ROBOTO 13 FIGMA ]") ) ]
+                +SVerticalBox::Slot().Padding(0.f, 64.f, 0.f, 4.f).AutoHeight() [ SNew(STextBlock).ColorAndOpacity(LabelColor).Text( NSLOCTEXT("StarshipGallery", "NormalBoldTextDesc", "ROBOTO BOLD 10 SLATE   [ ROBOTO 13 FIGMA ]") ) ]
                 +SVerticalBox::Slot().Padding(0.f, 4.f).AutoHeight() [ 
                     SNew(STextBlock)
                     .Font(FAppStyle::Get().GetFontStyle("NormalFontBold"))
@@ -536,7 +546,7 @@ public:
                 ]
 
                 // Small Bold
-                +SVerticalBox::Slot().Padding(0.f, 32.f, 0.f, 4.f).AutoHeight()[ SNew(STextBlock).ColorAndOpacity(LabelColor).Text(NSLOCTEXT("StarshipGallery", "SmallBoldTextDesc", "ROBOTO BOLD 8 SLATE   [ ROBOTO 11 FIGMA ]") ) ]
+                +SVerticalBox::Slot().Padding(0.f, 64.f, 0.f, 4.f).AutoHeight()[ SNew(STextBlock).ColorAndOpacity(LabelColor).Text(NSLOCTEXT("StarshipGallery", "SmallBoldTextDesc", "ROBOTO BOLD 8 SLATE   [ ROBOTO 11 FIGMA ]") ) ]
                 +SVerticalBox::Slot().Padding(0.f, 4.f).AutoHeight()[ 
                     SNew(STextBlock)
                     .Font(FAppStyle::Get().GetFontStyle("SmallFontBold"))
