SegmentedControl Widget

diff --git a/Engine/Source/Runtime/AppFramework/Private/Framework/Testing/SStarshipGallery.cpp b/Engine/Source/Runtime/AppFramework/Private/Framework/Testing/SStarshipGallery.cpp
--- a/Engine/Source/Runtime/AppFramework/Private/Framework/Testing/SStarshipGallery.cpp
+++ b/Engine/Source/Runtime/AppFramework/Private/Framework/Testing/SStarshipGallery.cpp
@@ -57,6 +57,7 @@
 #include "Widgets/Input/SSlider.h"
 #include "Widgets/Input/SComboBox.h"
 #include "Widgets/Input/SNumericEntryBox.h"
+#include "Widgets/Input/SSegmentedControl.h"
 
 
 #include "Brushes/SlateImageBrush.h"
@@ -952,115 +953,45 @@ public:
 
         ];
 
-        // Segmented Control
+
+        // Segmented Control - Simple version lets you specify text and brush
         NextSlot(WidgetGrid, LOCTEXT("SegmentedControl", "SegmentedControl")) 
         [
-            SNew(SHorizontalBox)
+        	SNew(SSegmentedControl<int32>)
+            .Value(0) // InitialValue
+            // .OnValueChanged_Lambda( [this] (int32 InValue) { SegmentedBoxChoice = InValue; } ) 
+            // .Value_Lambda( [this] { return SegmentedBoxChoice; } )  // Bound 
 
-            +SHorizontalBox::Slot()
-            .AutoWidth()
-            [
-                SNew(SCheckBox)
-                .Style( &FAppStyle::Get().GetWidgetStyle<FCheckBoxStyle>("SegmentedBoxLeft"))
-                .IsChecked_Lambda( [this] () -> ECheckBoxState { return SegmentedBoxChoice == 0 ? ECheckBoxState::Checked : ECheckBoxState::Unchecked; })
-                .OnCheckStateChanged_Lambda( [this] (ECheckBoxState InState) { if (InState == ECheckBoxState::Checked ) SegmentedBoxChoice = 0; } )
-                [
-                    LeftRightLabel("Icons.box-perspective")
-                ]
-            ]
+        	+SSegmentedControl<int32>::Slot(0)
+            .Icon(FAppStyle::Get().GetBrush("Icons.box-perspective"))
+            .Text(LOCTEXT("Box", "BOX"))
 
-            +SHorizontalBox::Slot()
-            .AutoWidth()
-            [
-                SNew(SCheckBox)
-                .Style( &FAppStyle::Get().GetWidgetStyle<FCheckBoxStyle>("SegmentedBoxCenter"))
-                .IsChecked_Lambda( [this] () -> ECheckBoxState { return SegmentedBoxChoice == 1 ? ECheckBoxState::Checked : ECheckBoxState::Unchecked; })
-                .OnCheckStateChanged_Lambda( [this] (ECheckBoxState InState) { if (InState == ECheckBoxState::Checked ) SegmentedBoxChoice = 1; } )
-                [
-                    LeftRightLabel("Icons.cylinder")
-                ]
-            ]
+            +SSegmentedControl<int32>::Slot(1)
+            .Icon(FAppStyle::Get().GetBrush("Icons.cylinder"))
+            .Text(LOCTEXT("Cylinder", "CYLINDER"))
 
-            +SHorizontalBox::Slot()
-            .AutoWidth()
-            [
-                SNew(SCheckBox)
-                .Style( &FAppStyle::Get().GetWidgetStyle<FCheckBoxStyle>("SegmentedBoxCenter"))
-                .IsChecked_Lambda( [this] () -> ECheckBoxState { return SegmentedBoxChoice == 2 ? ECheckBoxState::Checked : ECheckBoxState::Unchecked; })
-                .OnCheckStateChanged_Lambda( [this] (ECheckBoxState InState) { if (InState == ECheckBoxState::Checked ) SegmentedBoxChoice = 2; } )
-                [
-                    LeftRightLabel("Icons.pyramid")
-                ]
-            ]
+            +SSegmentedControl<int32>::Slot(2)
+            .Icon(FAppStyle::Get().GetBrush("Icons.pyramid"))
+            .Text(LOCTEXT("Pyramid", "PYRAMID"))
 
-            +SHorizontalBox::Slot()
-            .AutoWidth()
-            [
-                SNew(SCheckBox)
-                .Style( &FAppStyle::Get().GetWidgetStyle<FCheckBoxStyle>("SegmentedBoxRight"))
-                .IsChecked_Lambda( [this] () -> ECheckBoxState { return SegmentedBoxChoice == 3 ? ECheckBoxState::Checked : ECheckBoxState::Unchecked; })
-                .OnCheckStateChanged_Lambda( [this] (ECheckBoxState InState) { if (InState == ECheckBoxState::Checked ) SegmentedBoxChoice = 3; } )
-                [
-                    LeftRightLabel("Icons.sphere")
-                ]
-            ]
-        ];
+        	+SSegmentedControl<int32>::Slot(3)
+            .Icon(FAppStyle::Get().GetBrush("Icons.sphere"))
+            .Text(LOCTEXT("Sphere", "SPHERE"))
 
-        // Segmented Control
-        NextSlot(WidgetGrid, LOCTEXT("SegmentedControl", "SegmentedControl")) 
+        ];
+   
+        // Segmented Control - Explicitly Specify Contents
+        NextSlot(WidgetGrid, LOCTEXT("SegmentedControlWithChildren", "SegmentedControl Alt")) 
         [
-            SNew(SHorizontalBox)
-
-            +SHorizontalBox::Slot()
-            .AutoWidth()
-            [
-                SNew(SCheckBox)
-                .Style( &FAppStyle::Get().GetWidgetStyle<FCheckBoxStyle>("SegmentedBoxLeft"))
-                .IsChecked_Lambda( [this] () -> ECheckBoxState { return SegmentedBoxChoice == 0 ? ECheckBoxState::Checked : ECheckBoxState::Unchecked; })
-                .OnCheckStateChanged_Lambda( [this] (ECheckBoxState InState) { if (InState == ECheckBoxState::Checked ) SegmentedBoxChoice = 0; } )
-                [
-                    LeftRightLabel("Icons.box-perspective", LOCTEXT("Box", "BOX"))
-                ]
-            ]
-
-            +SHorizontalBox::Slot()
-            .AutoWidth()
-            [
-                SNew(SCheckBox)
-                .Style( &FAppStyle::Get().GetWidgetStyle<FCheckBoxStyle>("SegmentedBoxCenter"))
-                .IsChecked_Lambda( [this] () -> ECheckBoxState { return SegmentedBoxChoice == 1 ? ECheckBoxState::Checked : ECheckBoxState::Unchecked; })
-                .OnCheckStateChanged_Lambda( [this] (ECheckBoxState InState) { if (InState == ECheckBoxState::Checked ) SegmentedBoxChoice = 1; } )
-                [
-                    LeftRightLabel("Icons.cylinder", LOCTEXT("Cylinder", "CYLINDER"))
-                ]
-            ]
+            SNew(SSegmentedControl<int32>)
+            .Value(2)
 
-            +SHorizontalBox::Slot()
-            .AutoWidth()
-            [
-                SNew(SCheckBox)
-                .Style( &FAppStyle::Get().GetWidgetStyle<FCheckBoxStyle>("SegmentedBoxCenter"))
-                .IsChecked_Lambda( [this] () -> ECheckBoxState { return SegmentedBoxChoice == 2 ? ECheckBoxState::Checked : ECheckBoxState::Unchecked; })
-                .OnCheckStateChanged_Lambda( [this] (ECheckBoxState InState) { if (InState == ECheckBoxState::Checked ) SegmentedBoxChoice = 2; } )
-                [
-                    LeftRightLabel("Icons.pyramid", LOCTEXT("Pyramid", "PYRAMID"))
-                ]
-            ]
-
-            +SHorizontalBox::Slot()
-            .AutoWidth()
-            [
-                SNew(SCheckBox)
-                .Style( &FAppStyle::Get().GetWidgetStyle<FCheckBoxStyle>("SegmentedBoxRight"))
-                .IsChecked_Lambda( [this] () -> ECheckBoxState { return SegmentedBoxChoice == 3 ? ECheckBoxState::Checked : ECheckBoxState::Unchecked; })
-                .OnCheckStateChanged_Lambda( [this] (ECheckBoxState InState) { if (InState == ECheckBoxState::Checked ) SegmentedBoxChoice = 3; } )
-                [
-                    LeftRightLabel("Icons.sphere", LOCTEXT("Sphere", "SPHERE"))
-                ]
-            ]
+            +SSegmentedControl<int32>::Slot(0) [ LeftRightLabel("Icons.box-perspective") ]
+            +SSegmentedControl<int32>::Slot(1) [ LeftRightLabel("Icons.cylinder") ]
+            +SSegmentedControl<int32>::Slot(2) [ LeftRightLabel("Icons.pyramid") ]
+            +SSegmentedControl<int32>::Slot(3) [ LeftRightLabel("Icons.sphere") ]
         ];
 
-
         // SCheckBox
         NextSlot(WidgetGrid, LOCTEXT("SCheckBoxLabel", "Check Box"))
         [
diff --git a/Engine/Source/Runtime/Slate/Public/Widgets/Input/SSegmentedControl.h b/Engine/Source/Runtime/Slate/Public/Widgets/Input/SSegmentedControl.h
new file mode 100644
--- /dev/null
+++ b/Engine/Source/Runtime/Slate/Public/Widgets/Input/SSegmentedControl.h
@@ -0,0 +1,216 @@
+// Copyright Epic Games, Inc. All Rights Reserved.
+
+#pragma once
+
+#include "CoreMinimal.h"
+#include "Styling/AppStyle.h"
+#include "Styling/SegmentedControlStyle.h"
+#include "Widgets/SBoxPanel.h"
+#include "Widgets/Images/SImage.h"
+#include "Widgets/Input/SCheckBox.h"
+#include "Widgets/Layout/SSpacer.h"
+#include "Widgets/Text/STextBlock.h"
+
+/**
+ * A Slate Segmented Control is functionally similar to a group of Radio Buttons.
+ * Slots require a templated value to return when the segment is selected by the user.
+ * Users can specify text, icon or provide custom content to each Segment.
+ *
+ * Note: It is currently not possible to add segments after initialization 
+ *  (i.e. there is no AddSlot). 
+ */
+
+template< typename OptionType >
+class SSegmentedControl : public SCompoundWidget
+{
+
+public:
+
+	/** Stores the per-child info for this panel type */
+	template< typename OptionType >
+	struct FSlot : public TSlotBase<FSlot<OptionType>> 
+	{
+		FSlot(const OptionType& InValue) 
+		: TSlotBase<FSlot<OptionType>>() 
+		, _Text()
+		, _Icon(nullptr)
+		, _Value(InValue)
+		{ }
+
+		FSlot<OptionType>& Text(const TAttribute<FText>& InText)
+		{
+			_Text = InText;
+			return *this;
+		}
+
+		FSlot<OptionType>& Text(FText& InText)
+		{
+			_Text.SetValue(InText);
+			return *this;
+		}
+
+		FSlot<OptionType>& Icon(const TAttribute<const FSlateBrush*>& InBrush)
+		{
+			_Icon= InBrush;
+			return *this;
+		}
+
+	friend class SSegmentedControl<OptionType>;
+
+	protected:
+
+		TAttribute<FText> _Text;
+		TAttribute<const FSlateBrush*> _Icon;
+
+		OptionType _Value;
+	};
+
+	static FSlot<OptionType>& Slot(const OptionType& InValue)
+	{
+		return *(new FSlot<OptionType>(InValue));
+	}
+
+	DECLARE_DELEGATE_OneParam( FOnValueChanged, OptionType );
+
+	SLATE_BEGIN_ARGS( SSegmentedControl<OptionType> )
+	: _Style(&FAppStyle::Get().GetWidgetStyle<FSegmentedControlStyle>("SegmentedControl"))
+	{}
+		/** Slot type supported by this panel */
+		SLATE_SUPPORTS_SLOT(FSlot<OptionType>)
+
+		SLATE_STYLE_ARGUMENT(FSegmentedControlStyle, Style)
+
+		SLATE_ATTRIBUTE( OptionType, Value )
+
+		SLATE_EVENT( FOnValueChanged, OnValueChanged )
+		
+	SLATE_END_ARGS()
+
+	void Construct( const FArguments& InArgs )
+	{
+		check(InArgs._Style);
+
+		Style = InArgs._Style;
+
+		CurrentValue = InArgs._Value;
+		OnValueChanged = InArgs._OnValueChanged;
+
+		Children = InArgs.Slots;
+		RebuildChildren();
+	}
+
+	void RebuildChildren()
+	{
+		TSharedPtr<SHorizontalBox> HBox = SNew(SHorizontalBox);
+
+		const int32 NumSlots = Children.Num();
+		for ( int32 SlotIndex = 0; SlotIndex < NumSlots; ++SlotIndex )
+		{
+			TSharedRef<SWidget> Child = Children[SlotIndex]->GetWidget();
+			FSlot<OptionType>* ChildSlotPtr = Children[SlotIndex];
+			OptionType& ChildValue = ChildSlotPtr->_Value;			
+
+			if (Child == SNullWidget::NullWidget)
+			{
+				Child = SNew(SHorizontalBox)
+
+				+SHorizontalBox::Slot()
+				.AutoWidth()
+				.VAlign(VAlign_Center)
+				[
+					SNew(SImage)
+					.ColorAndOpacity(FSlateColor::UseForeground())
+					.Image_Lambda([ChildSlotPtr] () { return ChildSlotPtr->_Icon.Get(); })
+				]	
+
+				+SHorizontalBox::Slot()
+				.AutoWidth()
+				[
+					SNew(SSpacer)
+					.Size_Lambda( [ChildSlotPtr] () { return (ChildSlotPtr->_Icon.Get() != nullptr && !ChildSlotPtr->_Text.Get().IsEmpty()) ? FVector2D(8.0f, 1.0f) : FVector2D::ZeroVector; })
+				]
+
+				+SHorizontalBox::Slot()
+				.VAlign(VAlign_Center)
+				.Padding(0.0f, 0.5f, 0.f, 0.f)  // Compensate down for the baseline - helps when using all caps
+				.AutoWidth()
+				[
+					SNew(STextBlock)
+					.TextStyle( &FAppStyle::Get().GetWidgetStyle< FTextBlockStyle >( "ButtonText" ))
+					.Justification(ETextJustify::Center)
+					.Text_Lambda([ChildSlotPtr] () { return ChildSlotPtr->_Text.Get(); } ) 
+				];
+			}
+
+			HBox->AddSlot()
+			.AutoWidth()
+			[
+				SNew(SCheckBox)
+				.Style(SlotIndex == 0 ? &Style->FirstControlStyle : SlotIndex == (NumSlots - 1) ? &Style->LastControlStyle : &Style->ControlStyle)
+				.IsChecked(this, &SSegmentedControl::IsCurrentValue, ChildValue)
+				.OnCheckStateChanged(this, &SSegmentedControl::CommitValue, ChildValue)
+				[
+					Child
+				]
+			];
+		}
+
+		ChildSlot
+		[
+			HBox.ToSharedRef()
+		];
+
+	}
+
+	// Slot Management
+	FSlot<OptionType>& AddSlot(const OptionType& InValue)
+	{
+		FSlot<OptionType>& NewSlot = *(new FSlot<OptionType>(InValue));
+
+		Children.Add( &NewSlot );
+		RebuildChildren();
+
+		return NewSlot;	
+	}
+
+	int32 NumSlots() const
+	{
+		return Children.Num();
+	}
+
+	/** See the Value attribute */
+	OptionType& GetValue() const { return CurrentValue.Get(); }
+	void SetValue(const TAttribute<OptionType>& InValue) 
+	{
+		CurrentValue = InValue; 
+	}
+
+	ECheckBoxState IsCurrentValue(OptionType InValue) const
+	{
+		return InValue == CurrentValue.Get() ? ECheckBoxState::Checked : ECheckBoxState::Unchecked;
+	}	
+
+	void CommitValue(const ECheckBoxState InCheckState, OptionType InValue) 
+	{
+		if (InCheckState == ECheckBoxState::Checked)
+		{
+			// don't overwrite the bound attribute, but still notify that a
+			if (!CurrentValue.IsBound())
+			{
+				CurrentValue.Set(InValue);
+			}
+
+			OnValueChanged.ExecuteIfBound( InValue );
+		}
+	}
+
+
+protected:
+	TArray<FSlot<OptionType>*> Children;
+
+	FOnValueChanged OnValueChanged;
+
+	TAttribute<OptionType> CurrentValue;
+
+	const FSegmentedControlStyle* Style;
+};
diff --git a/Engine/Source/Runtime/SlateCore/Private/Styling/SegmentedControlStyle.cpp b/Engine/Source/Runtime/SlateCore/Private/Styling/SegmentedControlStyle.cpp
new file mode 100644
--- /dev/null
+++ b/Engine/Source/Runtime/SlateCore/Private/Styling/SegmentedControlStyle.cpp
@@ -0,0 +1,22 @@
+// Copyright Epic Games, Inc. All Rights Reserved.
+
+#include "Styling/SegmentedControlStyle.h"
+
+const FName FSegmentedControlStyle::TypeName( TEXT("FSegmentedControlStyle") );
+
+FSegmentedControlStyle::FSegmentedControlStyle()
+{
+}
+
+void FSegmentedControlStyle::GetResources( TArray< const FSlateBrush* >& OutBrushes ) const
+{
+	ControlStyle.GetResources(OutBrushes);
+	FirstControlStyle.GetResources(OutBrushes);
+	LastControlStyle.GetResources(OutBrushes);
+}
+
+const FSegmentedControlStyle& FSegmentedControlStyle::GetDefault()
+{
+	static FSegmentedControlStyle Default;
+	return Default;
+}
\ No newline at end of file
diff --git a/Engine/Source/Runtime/SlateCore/Private/Styling/StarshipCoreStyle.cpp b/Engine/Source/Runtime/SlateCore/Private/Styling/StarshipCoreStyle.cpp
--- a/Engine/Source/Runtime/SlateCore/Private/Styling/StarshipCoreStyle.cpp
+++ b/Engine/Source/Runtime/SlateCore/Private/Styling/StarshipCoreStyle.cpp
@@ -13,6 +13,7 @@
 #include "Fonts/LegacySlateFontInfoCache.h"
 #include "Styling/SlateStyleMacros.h"
 #include "Styling/ToolBarStyle.h"
+#include "Styling/SegmentedControlStyle.h"
 #include "Styling/StyleColors.h"
 
 // This is to fix the issue that SlateStyleMacros like IMAGE_BRUSH look for RootToContentDir but Style->RootToContentDir is how the core style is set up
@@ -1305,10 +1306,11 @@ void FStarshipCoreStyle::SetupCheckboxStyles(TSharedRef<FStyle>& Style)
 		.SetCheckedPressedForegroundColor(FStyleColors::Primary)
 		.SetPadding(DefaultMargins);
 
-	Style->Set("SegmentedBoxLeft", SegmentedBoxLeft);
-	Style->Set("SegmentedBoxCenter", SegmentedBoxCenter);
-	Style->Set("SegmentedBoxRight", SegmentedBoxRight);
-
+	Style->Set("SegmentedControl", FSegmentedControlStyle()
+		.SetControlStyle(SegmentedBoxCenter)
+		.SetFirstControlStyle(SegmentedBoxLeft)
+		.SetLastControlStyle(SegmentedBoxRight)
+	);
 
 
 	/* Style for a toggleable button that mimics the coloring and look of a Table Row */
diff --git a/Engine/Source/Runtime/SlateCore/Public/Styling/SegmentedControlStyle.h b/Engine/Source/Runtime/SlateCore/Public/Styling/SegmentedControlStyle.h
new file mode 100644
--- /dev/null
+++ b/Engine/Source/Runtime/SlateCore/Public/Styling/SegmentedControlStyle.h
@@ -0,0 +1,62 @@
+// Copyright Epic Games, Inc. All Rights Reserved.
+
+#pragma once
+
+#include "CoreMinimal.h"
+#include "UObject/ObjectMacros.h"
+#include "Styling/SlateTypes.h"
+
+#include "SegmentedControlStyle.generated.h"
+
+/**
+ * Represents the appearance of an SSegmentedControl
+ */
+USTRUCT(BlueprintType)
+struct SLATECORE_API FSegmentedControlStyle: public FSlateWidgetStyle
+{
+	GENERATED_USTRUCT_BODY()
+
+	FSegmentedControlStyle();
+
+	virtual ~FSegmentedControlStyle() {}
+
+	virtual void GetResources( TArray< const FSlateBrush* >& OutBrushes ) const override;
+
+	static const FName TypeName;
+	virtual const FName GetTypeName() const override { return TypeName; };
+
+	static const FSegmentedControlStyle& GetDefault();
+
+	/**
+	 * The style to use for our Center Control
+	 */	
+	UPROPERTY(EditAnywhere, BlueprintReadWrite, Category=Appearance)
+	FCheckBoxStyle ControlStyle;
+	FSegmentedControlStyle& SetControlStyle( const FCheckBoxStyle& InStyle ){ ControlStyle = InStyle; return *this; }
+
+	/**
+	 * The style to use for our Left Control
+	 */	
+	UPROPERTY(EditAnywhere, BlueprintReadWrite, Category=Appearance)
+	FCheckBoxStyle FirstControlStyle;
+	FSegmentedControlStyle& SetFirstControlStyle( const FCheckBoxStyle& InStyle ){ FirstControlStyle = InStyle; return *this; }
+
+	/**
+	 * The style to use for our Left Control
+	 */	
+	UPROPERTY(EditAnywhere, BlueprintReadWrite, Category=Appearance)
+	FCheckBoxStyle LastControlStyle;
+	FSegmentedControlStyle& SetLastControlStyle( const FCheckBoxStyle& InStyle ){ LastControlStyle = InStyle; return *this; }
+
+	/**
+	 * Unlinks all colors in this style.
+	 * @see FSlateColor::Unlink
+	 */
+	void UnlinkColors()
+	{
+		ControlStyle.UnlinkColors();
+		FirstControlStyle.UnlinkColors();
+		LastControlStyle.UnlinkColors();
+	}
+
+};
\ No newline at end of file
