Removing unused Quixel Icons, update SlateViewer Icon Manefest

diff --git a/Engine/Source/Runtime/AppFramework/Private/Framework/Testing/SStarshipGallery.cpp b/Engine/Source/Runtime/AppFramework/Private/Framework/Testing/SStarshipGallery.cpp
--- a/Engine/Source/Runtime/AppFramework/Private/Framework/Testing/SStarshipGallery.cpp
+++ b/Engine/Source/Runtime/AppFramework/Private/Framework/Testing/SStarshipGallery.cpp
@@ -123,7 +123,6 @@ public:
                     return SNew(SDockTab)
                     .TabRole(ETabRole::PanelTab)
                     .ContentPadding(0)
-                    // .Icon(DynamicBrushes.Add_GetRef(TUniquePtr<FSlateDynamicImageBrush>(new FSlateDynamicImageBrush(FName("../../Content/Slate/Starship/Icons/Quixel/DesignTools/layer-liquid_32x.png"), FVector2D(16.f, 16.f)))).Get())
                     [
                         ConstructColorsGallery()
                     ];
@@ -138,7 +137,6 @@ public:
                     return SNew(SDockTab)
                     .TabRole(ETabRole::PanelTab)
                     .ContentPadding(0)
-                    // .Icon(DynamicBrushes.Add_GetRef(TUniquePtr<FSlateDynamicImageBrush>(new FSlateDynamicImageBrush(FName("../../Content/Slate/Starship/Icons/Quixel/Common/tile_32x.png"), FVector2D(16.f, 16.f)))).Get())
                      [
                         ConstructTextGallery()
                     ];
@@ -154,7 +152,6 @@ public:
                     return SNew(SDockTab)
                     .TabRole(ETabRole::PanelTab)
                     .ContentPadding(0)
-                    // .Icon(DynamicBrushes.Add_GetRef(TUniquePtr<FSlateDynamicImageBrush>(new FSlateDynamicImageBrush(FName("../../Content/Slate/Starship/Icons/Quixel/Common/tile_32x.png"), FVector2D(16.f, 16.f)))).Get())
                      [
                         ConstructIconsGallery()
                     ];
@@ -171,7 +168,6 @@ public:
                     .TabRole(ETabRole::PanelTab)
                     .ContentPadding(0) 
                     .ForegroundColor(FSlateColor::UseStyle())
-                    // .Icon(DynamicBrushes.Add_GetRef(TUniquePtr<FSlateDynamicImageBrush>(new FSlateDynamicImageBrush(FName("../../Content/Slate/Starship/Icons/Quixel/Common/gift_32x.png"), FVector2D(16.f, 16.f)))).Get())
                     [
                         ConstructWidgetGallery()
                     ];
@@ -467,15 +463,21 @@ public:
             .Padding(48)
             [
                 SNew(SVerticalBox)
-                +SVerticalBox::Slot().AutoHeight()[ GenerateIconLibrary(NSLOCTEXT("StarshipGallery", "GeneralIconTitle", "General"), "Content/Editor/Slate/Icons/GeneralTools")]
+                +SVerticalBox::Slot().AutoHeight()[ GenerateIconLibrarySVG(NSLOCTEXT("StarshipGallery", "Starship Common", "Common"), "Content/Slate/Starship/Common")]
+                
+                +SVerticalBox::Slot().AutoHeight()[ GenerateIconLibrarySVG(NSLOCTEXT("StarshipGallery", "SceneOutliner", "SceneOutliner"), "Content/Editor/Slate/Starship/SceneOutliner")]
+                +SVerticalBox::Slot().AutoHeight()[ GenerateIconLibrarySVG(NSLOCTEXT("StarshipGallery", "LevelEditor", "LevelEditor"), "Content/Editor/Slate/Starship/LevelEditor/Menus")]
+                +SVerticalBox::Slot().AutoHeight()[ GenerateIconLibrarySVG(NSLOCTEXT("StarshipGallery", "MainToolbar", "MainToolbar"), "Content/Editor/Slate/Starship/MainToolbar")]
+                +SVerticalBox::Slot().AutoHeight()[ GenerateIconLibrarySVG(NSLOCTEXT("StarshipGallery", "FileMenu", "FileMenu"), "Content/Editor/Slate/Starship/Menus/File")]
+                +SVerticalBox::Slot().AutoHeight()[ GenerateIconLibrarySVG(NSLOCTEXT("StarshipGallery", "EditMenu", "EditMenu"), "Content/Editor/Slate/Starship/Menus/Edit")]
+                +SVerticalBox::Slot().AutoHeight()[ GenerateIconLibrarySVG(NSLOCTEXT("StarshipGallery", "HelpMenu", "HelpMenu"), "Content/Editor/Slate/Starship/Menus/Help")]
+
                 +SVerticalBox::Slot().AutoHeight()[ GenerateIconLibrary(NSLOCTEXT("StarshipGallery", "PaintIconTitle", "Paint"), "Content/Editor/Slate/Icons/Paint")]
                 +SVerticalBox::Slot().AutoHeight()[ GenerateIconLibrary(NSLOCTEXT("StarshipGallery", "LandscapeIconTitle", "Landscape"), "Content/Editor/Slate/Icons/Landscape")]
                 +SVerticalBox::Slot().AutoHeight()[ GenerateIconLibrary(NSLOCTEXT("StarshipGallery", "ModelingIconTitle", "Modeling"), "/Plugins/Experimental/ModelingToolsEditorMode/Content/Icons")]
                 +SVerticalBox::Slot().AutoHeight()[ GenerateIconLibrary(NSLOCTEXT("StarshipGallery", "FractureIconTitle", "Fracture"), "/Plugins/Experimental/ChaosEditor/Content")]
                 +SVerticalBox::Slot().AutoHeight()[ GenerateIconLibrary(NSLOCTEXT("StarshipGallery", "CurveEditorIconTitle", "CurveEditor"), "Content/Editor/Slate/GenericCurveEditor/Icons")]
-                // +SVerticalBox::Slot().AutoHeight()[ GenerateIconLibrary(NSLOCTEXT("StarshipGallery", "QuixelIconTitle", "Quixel"), "Content/Slate/Starship/Icons/Quixel")]
-                +SVerticalBox::Slot().AutoHeight()[ GenerateIconLibrarySVG(NSLOCTEXT("StarshipGallery", "QuixelIcons", "Quixel"), "Content/Slate/Starship/Common")]
-                +SVerticalBox::Slot().AutoHeight()[ GenerateIconLibrarySVG(NSLOCTEXT("StarshipGallery", "SceneOutliner", "SceneOutliner"), "Content/Slate/Starship/SceneOutliner")]
+                +SVerticalBox::Slot().AutoHeight()[ GenerateIconLibrary(NSLOCTEXT("StarshipGallery", "GeneralIconTitle", "General"), "Content/Editor/Slate/Icons/GeneralTools")]
             ]
         ];
     }
