Fixed some of the color theme viewer text not fitting in the button and the list moving off screen if too many items were added to it

diff --git a/Engine/Source/Runtime/AppFramework/Private/Widgets/Colors/SColorThemes.cpp b/Engine/Source/Runtime/AppFramework/Private/Widgets/Colors/SColorThemes.cpp
--- a/Engine/Source/Runtime/AppFramework/Private/Widgets/Colors/SColorThemes.cpp
+++ b/Engine/Source/Runtime/AppFramework/Private/Widgets/Colors/SColorThemes.cpp
@@ -919,10 +919,12 @@ void SColorThemesViewer::Construct(const FArguments& InArgs)
 	this->ChildSlot
 	[
 		SNew(SVerticalBox)
-		+SVerticalBox::Slot().AutoHeight()
+		+SVerticalBox::Slot()
+			.AutoHeight()
+			.MaxHeight(455.0f) // Roughly matches the size of the color picker at max height
 		[
 			SNew(SBox)
-			.WidthOverride(320)
+			.WidthOverride(360)
 			[
 				SAssignNew(ColorThemeList, SListView< TSharedPtr<FColorTheme> >)
 					.ItemHeight(32)
@@ -930,11 +932,14 @@ void SColorThemesViewer::Construct(const FArguments& InArgs)
 					.OnGenerateRow(this, &SColorThemesViewer::OnGenerateColorThemeBars)
 			]
 		]
-		+SVerticalBox::Slot().AutoHeight() .Padding(0,15,0,0)
+		+SVerticalBox::Slot()
+			.AutoHeight()
+			.Padding(0,15,0,0)
 		[
 			SAssignNew(Menu, SBorder)
 		]
-		+ SVerticalBox::Slot().AutoHeight()
+		+ SVerticalBox::Slot()
+			.AutoHeight()
 		[
 			SAssignNew(ErrorText, SErrorText)
 				.Visibility(this, &SColorThemesViewer::OnGetErrorTextVisibility)
