STextBlockStyle added TransformPolicy

- Update Button Text Styles

diff --git a/Engine/Source/Runtime/AppFramework/Private/Framework/Testing/SStarshipGallery.cpp b/Engine/Source/Runtime/AppFramework/Private/Framework/Testing/SStarshipGallery.cpp
--- a/Engine/Source/Runtime/AppFramework/Private/Framework/Testing/SStarshipGallery.cpp
+++ b/Engine/Source/Runtime/AppFramework/Private/Framework/Testing/SStarshipGallery.cpp
@@ -464,7 +464,7 @@ public:
             [
                 SNew(SVerticalBox)
                 +SVerticalBox::Slot().AutoHeight()[ GenerateIconLibrarySVG(NSLOCTEXT("StarshipGallery", "Starship Common", "Common"), "Content/Slate/Starship/Common")]
-                
+
                 +SVerticalBox::Slot().AutoHeight()[ GenerateIconLibrarySVG(NSLOCTEXT("StarshipGallery", "SceneOutliner", "SceneOutliner"), "Content/Editor/Slate/Starship/SceneOutliner")]
                 +SVerticalBox::Slot().AutoHeight()[ GenerateIconLibrarySVG(NSLOCTEXT("StarshipGallery", "LevelEditor", "LevelEditor"), "Content/Editor/Slate/Starship/LevelEditor/Menus")]
                 +SVerticalBox::Slot().AutoHeight()[ GenerateIconLibrarySVG(NSLOCTEXT("StarshipGallery", "MainToolbar", "MainToolbar"), "Content/Editor/Slate/Starship/MainToolbar")]
@@ -639,7 +639,7 @@ public:
         };
 
 
-        auto LeftRightLabel = [](const FName& InIconName = FName(), const FText& InLabel = FText::GetEmpty()) -> TSharedRef<SWidget>
+        auto LeftRightLabel = [](const FName& InIconName = FName(), const FText& InLabel = FText::GetEmpty(), const FName& InTextStyle = TEXT("ButtonText")) -> TSharedRef<SWidget>
         {
         	TSharedPtr<SHorizontalBox> HBox = SNew(SHorizontalBox);
             float Space = InIconName.IsNone() ? 0.0f : 8.f;
@@ -664,7 +664,7 @@ public:
 	            .AutoWidth()
 	            [
 	                SNew(STextBlock)
-	                .TextStyle( &FAppStyle::Get().GetWidgetStyle< FTextBlockStyle >( "ButtonText" ))
+	                .TextStyle( &FAppStyle::Get().GetWidgetStyle< FTextBlockStyle >( InTextStyle ))
 	                .Justification(ETextJustify::Center)
 	                .Text(InLabel)
 	            ];
@@ -731,24 +731,27 @@ public:
 
                 SNew(SButton)
                 .ButtonStyle( &FAppStyle::Get().GetWidgetStyle< FButtonStyle >( "PrimaryButton" ) )
+                .TextStyle( &FAppStyle::Get().GetWidgetStyle< FTextBlockStyle >("ButtonText"))
                 [
-                    LeftRightLabel(NAME_None, LOCTEXT("PrimaryButtonExampleLabel", "PRIMARY BUTTON"))
+                    LeftRightLabel(NAME_None, LOCTEXT("PrimaryButtonExampleLabel", "Primary Button"))
                 ]
             ]
 
             +SHorizontalBox::Slot()
             .AutoWidth()
             .Padding(8.f, 0.0f)
+            .VAlign(VAlign_Center)
             [
                 SNew(SButton)
                 .ButtonStyle( &FAppStyle::Get().GetWidgetStyle< FButtonStyle >( "PrimaryButton" ) )
                 [
-                    LeftRightLabel("Icons.box-perspective", LOCTEXT("PrimaryButtonExampleLabel", "PRIMARY BUTTON"))
+                    LeftRightLabel("Icons.box-perspective", LOCTEXT("PrimaryButtonExampleLabel", "Primary Button"))
                 ]
             ]
 
             +SHorizontalBox::Slot()
             .AutoWidth()
+            .VAlign(VAlign_Center)
             .Padding(8.f, 0.0f)
             [
                 SNew(SButton)
@@ -773,7 +776,7 @@ public:
                 SNew(SButton)
                 .ButtonStyle( &FAppStyle::Get().GetWidgetStyle< FButtonStyle >( "Button" ) )
                 [
-                    LeftRightLabel(NAME_None, LOCTEXT("ButtonExampleLabel", "NORMAL BUTTON"))
+                    LeftRightLabel(NAME_None, LOCTEXT("ButtonExampleLabel", "Normal Button"))
                 ]
             ]
 
@@ -784,7 +787,7 @@ public:
                 SNew(SButton)
                 .ButtonStyle( &FAppStyle::Get().GetWidgetStyle< FButtonStyle >( "Button" ) )
                 [
-                    LeftRightLabel("Icons.box-perspective", LOCTEXT("ButtonExampleLabel", "NORMAL BUTTON"))
+                    LeftRightLabel("Icons.box-perspective", LOCTEXT("ButtonExampleLabel", "Normal Button"))
                 ]
             ]
 
@@ -817,7 +820,7 @@ public:
                 SNew(SButton)
                 .ButtonStyle( &FAppStyle::Get().GetWidgetStyle< FButtonStyle >( "SimpleButton" ) )
                 [
-                    LeftRightLabel(NAME_None, LOCTEXT("TextButtonExampleLabel", "SIMPLE BUTTON"))
+                    LeftRightLabel(NAME_None, LOCTEXT("TextButtonExampleLabel", "Simple Button"))
                 ]
             ]
 
@@ -828,7 +831,7 @@ public:
                 SNew(SButton)
                 .ButtonStyle( &FAppStyle::Get().GetWidgetStyle< FButtonStyle >( "SimpleButton" ) )
                 [
-                    LeftRightLabel("Icons.box-perspective", LOCTEXT("TextButtonExampleLabel", "SIMPLE BUTTON")) 
+                    LeftRightLabel("Icons.box-perspective", LOCTEXT("TextButtonExampleLabel", "Simple Button")) 
                 ]
             ]
 
@@ -845,6 +848,92 @@ public:
 
         ];
 
+        // SButton Dialog Primary Rounded
+        NextSlot(WidgetGrid, LOCTEXT("SButtonDialogPrimaryExampleLabelRounded", "Dialog Primary Button"))
+        [
+            SNew(SHorizontalBox)
+            +SHorizontalBox::Slot()
+            .AutoWidth()
+            .Padding(8.f, 0.0f)
+            .VAlign(VAlign_Center)
+            [
+
+                SNew(SButton)
+                .ButtonStyle( &FAppStyle::Get().GetWidgetStyle< FButtonStyle >( "PrimaryButton" ) )
+                .TextStyle( &FAppStyle::Get().GetWidgetStyle< FTextBlockStyle >("DialogButtonText"))
+                [
+                    LeftRightLabel(NAME_None, LOCTEXT("DialogPrimaryButtonExampleLabel", "Dialog Primary"), "DialogButtonText")
+                ]
+            ]
+
+            +SHorizontalBox::Slot()
+            .AutoWidth()
+            .Padding(8.f, 0.0f)
+            .VAlign(VAlign_Center)
+            [
+                SNew(SButton)
+                .ButtonStyle( &FAppStyle::Get().GetWidgetStyle< FButtonStyle >( "PrimaryButton" ) )
+                [
+                    LeftRightLabel("Icons.box-perspective", LOCTEXT("DialogPrimaryButtonExampleLabel", "Dialog Primary"), "DialogButtonText")
+                ]
+            ]
+
+            +SHorizontalBox::Slot()
+            .AutoWidth()
+            .Padding(8.f, 0.0f)
+            .VAlign(VAlign_Center)
+            [
+                SNew(SButton)
+                .ButtonStyle( &FAppStyle::Get().GetWidgetStyle< FButtonStyle >( "PrimaryButton" ) )
+                .VAlign(VAlign_Center)
+                [
+                    LeftRightLabel("Icons.box-perspective")
+                ]
+            ]
+        ];
+
+        // SButton Dialog 
+        NextSlot(WidgetGrid, LOCTEXT("SButtonLabelRounded", "Button"))
+        [
+
+            SNew(SHorizontalBox)
+            +SHorizontalBox::Slot()
+            .AutoWidth()
+            .Padding(8.f, 0.0f)
+            .VAlign(VAlign_Center)
+            [
+                SNew(SButton)
+                .ButtonStyle( &FAppStyle::Get().GetWidgetStyle< FButtonStyle >( "Button" ) )
+                [
+                    LeftRightLabel(NAME_None, LOCTEXT("DialogButtonExampleLabel", "Dialog Button"), "DialogButtonText")
+                ]
+            ]
+
+            +SHorizontalBox::Slot()
+            .AutoWidth()
+            .Padding(8.f, 0.0f)
+            [
+                SNew(SButton)
+                .ButtonStyle( &FAppStyle::Get().GetWidgetStyle< FButtonStyle >( "Button" ) )
+                [
+                    LeftRightLabel("Icons.box-perspective", LOCTEXT("DialogButtonExampleLabel", "Dialog Button"), "DialogButtonText")
+                ]
+            ]
+
+            +SHorizontalBox::Slot()
+            .AutoWidth()
+            .Padding(8.f, 0.0f)
+            [
+                SNew(SButton)
+                .ButtonStyle( &FAppStyle::Get().GetWidgetStyle< FButtonStyle >( "Button" ) )
+                .VAlign(VAlign_Center)
+                [
+                    LeftRightLabel("Icons.box-perspective")
+                ]
+            ]
+
+        ];
+
 
         // Toggle Button 
         NextSlot(WidgetGrid, LOCTEXT("ToggleButtonAlt", "Toggle Button Alt"))
diff --git a/Engine/Source/Runtime/Slate/Private/Widgets/Text/STextBlock.cpp b/Engine/Source/Runtime/Slate/Private/Widgets/Text/STextBlock.cpp
--- a/Engine/Source/Runtime/Slate/Private/Widgets/Text/STextBlock.cpp
+++ b/Engine/Source/Runtime/Slate/Private/Widgets/Text/STextBlock.cpp
@@ -42,7 +42,11 @@ void STextBlock::Construct( const FArguments& InArgs )
 	WrapTextAt = InArgs._WrapTextAt;
 	AutoWrapText = InArgs._AutoWrapText;
 	WrappingPolicy = InArgs._WrappingPolicy;
-	TransformPolicy = InArgs._TransformPolicy;
+
+	if (InArgs._TransformPolicy.IsSet()) 
+	{
+		TransformPolicy = InArgs._TransformPolicy;
+	}
 	Margin = InArgs._Margin;
 	LineHeightPercentage = InArgs._LineHeightPercentage;
 	Justification = InArgs._Justification;
@@ -276,7 +280,7 @@ FVector2D STextBlock::ComputeDesiredSize(float LayoutScaleMultiplier) const
 	{
 		// ComputeDesiredSize will also update the text layout cache if required
 		const FVector2D TextSize = TextLayoutCache->ComputeDesiredSize(
-			FSlateTextBlockLayout::FWidgetArgs(BoundText, HighlightText, WrapTextAt, AutoWrapText, WrappingPolicy, TransformPolicy, Margin, LineHeightPercentage, Justification),
+			FSlateTextBlockLayout::FWidgetArgs(BoundText, HighlightText, WrapTextAt, AutoWrapText, WrappingPolicy, GetTransformPolicy(), Margin, LineHeightPercentage, Justification),
 			LayoutScaleMultiplier, GetComputedTextStyle()
 		);
 
@@ -403,6 +407,15 @@ void STextBlock::SetTransformPolicy(const TAttribute<ETextTransformPolicy>& InTr
 	}
 }
 
+ETextTransformPolicy STextBlock::GetTransformPolicy() const
+{
+	if (TransformPolicy.IsSet())
+	{
+		return TransformPolicy.Get();
+	}
+	return TextStyle.TransformPolicy;
+}
+
 void STextBlock::SetShadowOffset(const TAttribute<FVector2D>& InShadowOffset)
 {
 	if(!ShadowOffset.IdenticalTo(InShadowOffset))
diff --git a/Engine/Source/Runtime/Slate/Public/Framework/Text/TextLayout.h b/Engine/Source/Runtime/Slate/Public/Framework/Text/TextLayout.h
--- a/Engine/Source/Runtime/Slate/Public/Framework/Text/TextLayout.h
+++ b/Engine/Source/Runtime/Slate/Public/Framework/Text/TextLayout.h
@@ -46,21 +46,6 @@ namespace ETextJustify
 	};
 }
 
-/**
- * The different methods that can be used if a word is too long to be broken by the default line-break iterator.
- */
-UENUM(BlueprintType)
-enum class ETextTransformPolicy : uint8
-{
-	/** No transform, just use the given text as-is */
-	None = 0,
-
-	/** Convert the text to lowercase for display */
-	ToLower,
-
-	/** Convert the text to uppercase for display */
-	ToUpper,
-};
 
 /** 
  * The different methods that can be used if a word is too long to be broken by the default line-break iterator.
diff --git a/Engine/Source/Runtime/Slate/Public/Widgets/Input/SButton.h b/Engine/Source/Runtime/Slate/Public/Widgets/Input/SButton.h
--- a/Engine/Source/Runtime/Slate/Public/Widgets/Input/SButton.h
+++ b/Engine/Source/Runtime/Slate/Public/Widgets/Input/SButton.h
@@ -35,7 +35,7 @@ public:
 	SLATE_BEGIN_ARGS( SButton )
 		: _Content()
 		, _ButtonStyle( &FCoreStyle::Get().GetWidgetStyle< FButtonStyle >( "Button" ) )
-		, _TextStyle( &FCoreStyle::Get().GetWidgetStyle< FTextBlockStyle >("NormalText") )
+		, _TextStyle( &FCoreStyle::Get().GetWidgetStyle< FTextBlockStyle >("ButtonText") )
 		, _HAlign( HAlign_Fill )
 		, _VAlign( VAlign_Fill )
 		, _ContentPadding(FMargin(4.0, 2.0))
@@ -46,7 +46,7 @@ public:
 		, _DesiredSizeScale( FVector2D(1,1) )
 		, _ContentScale( FVector2D(1,1) )
 		, _ButtonColorAndOpacity(FLinearColor::White)
-		, _ForegroundColor( FCoreStyle::Get().GetSlateColor( "InvertedForeground" ) )
+		, _ForegroundColor(FSlateColor::UseStyle())
 		, _IsFocusable( true )
 		{
 		}
@@ -124,6 +124,12 @@ public:
 	/** @return An image that represents this button's border*/
 	virtual const FSlateBrush* GetBorder() const;
 
+	/** @return the Foreground color that this widget sets; unset options if the widget does not set a foreground color */
+	virtual FSlateColor GetForegroundColor() const;
+
+	/** @return the Foreground color that this widget sets when this widget or any of its ancestors are disabled; unset options if the widget does not set a foreground color */
+	virtual FSlateColor GetDisabledForegroundColor() const;
+
 	/**
 	 * Returns true if this button is currently pressed
 	 *
diff --git a/Engine/Source/Runtime/Slate/Public/Widgets/Text/STextBlock.h b/Engine/Source/Runtime/Slate/Public/Widgets/Text/STextBlock.h
--- a/Engine/Source/Runtime/Slate/Public/Widgets/Text/STextBlock.h
+++ b/Engine/Source/Runtime/Slate/Public/Widgets/Text/STextBlock.h
@@ -55,7 +55,7 @@ public:
 		, _WrapTextAt(0.0f)
 		, _AutoWrapText(false)
 		, _WrappingPolicy(ETextWrappingPolicy::DefaultWrapping)
-		, _TransformPolicy(ETextTransformPolicy::None)
+		, _TransformPolicy()
 		, _Margin()
 		, _LineHeightPercentage(1.0f)
 		, _Justification(ETextJustify::Left)
@@ -222,6 +222,7 @@ public:
 
 	/** Set TransformPolicy attribute */
 	void SetTransformPolicy(const TAttribute<ETextTransformPolicy>& InTransformPolicy);
+	ETextTransformPolicy GetTransformPolicy() const;
 
 	/** See ShadowOffset attribute */
 	void SetShadowOffset(const TAttribute<FVector2D>& InShadowOffset);
diff --git a/Engine/Source/Runtime/SlateCore/Private/Styling/SlateTypes.cpp b/Engine/Source/Runtime/SlateCore/Private/Styling/SlateTypes.cpp
--- a/Engine/Source/Runtime/SlateCore/Private/Styling/SlateTypes.cpp
+++ b/Engine/Source/Runtime/SlateCore/Private/Styling/SlateTypes.cpp
@@ -84,6 +84,7 @@ FTextBlockStyle::FTextBlockStyle()
 	, ShadowColorAndOpacity(FLinearColor::Black)
 	, SelectedBackgroundColor(FSlateColor::UseForeground())
 	, HighlightColor(FLinearColor::Black)
+	, TransformPolicy(ETextTransformPolicy::None)
 {
 }
 
diff --git a/Engine/Source/Runtime/SlateCore/Private/Styling/StarshipCoreStyle.cpp b/Engine/Source/Runtime/SlateCore/Private/Styling/StarshipCoreStyle.cpp
--- a/Engine/Source/Runtime/SlateCore/Private/Styling/StarshipCoreStyle.cpp
+++ b/Engine/Source/Runtime/SlateCore/Private/Styling/StarshipCoreStyle.cpp
@@ -51,7 +51,7 @@ namespace CoreStyleConstants
 	// Common Margins
 	const FMargin DefaultMargins(8.f, 4.f);
 	// Buttons already have a built in (4., 2.) padding - adding to that a little
-	const FMargin ButtonMargins(8.f, 2.f);
+	const FMargin ButtonMargins(20.f, 2.f, 20.f, 3.f);
 
 	const float InputFocusRadius = 4.f;
 	const float InputFocusThickness = 1.0f;
@@ -1081,11 +1081,13 @@ void FStarshipCoreStyle::SetupButtonStyles(TSharedRef<FStyle>& Style)
 
 		const FTextBlockStyle& NormalText = Style->GetWidgetStyle<FTextBlockStyle>("NormalText");
 
-		Style->Set("ButtonText", FTextBlockStyle(NormalText).SetFont(FStyleFonts::Get().SmallBold));
+		Style->Set("ButtonText", FTextBlockStyle(NormalText).SetFont(FStyleFonts::Get().NormalBold));
 
 		Style->Set("PrimaryButton", PrimaryButton);
 		Style->Set("Button", Button);
 		Style->Set("SimpleButton", SimpleButton);
+
+		Style->Set("DialogButtonText", FTextBlockStyle(NormalText).SetFont(FStyleFonts::Get().NormalBold).SetTransformPolicy(ETextTransformPolicy::ToUpper));
 	}
 }
 
diff --git a/Engine/Source/Runtime/SlateCore/Public/Styling/SlateTypes.h b/Engine/Source/Runtime/SlateCore/Public/Styling/SlateTypes.h
--- a/Engine/Source/Runtime/SlateCore/Public/Styling/SlateTypes.h
+++ b/Engine/Source/Runtime/SlateCore/Public/Styling/SlateTypes.h
@@ -264,6 +264,22 @@ struct TStructOpsTypeTraits<FCheckBoxStyle> : public TStructOpsTypeTraitsBase2<F
 };
 #endif
 
+/**
+ * The different methods that can be used if a word is too long to be broken by the default line-break iterator.
+ */
+UENUM(BlueprintType)
+enum class ETextTransformPolicy : uint8
+{
+	/** No transform, just use the given text as-is */
+	None = 0,
+
+	/** Convert the text to lowercase for display */
+	ToLower,
+
+	/** Convert the text to uppercase for display */
+	ToUpper,
+};
+
 /**
  * Represents the appearance of an STextBlock
  */
@@ -340,6 +356,11 @@ struct SLATECORE_API FTextBlockStyle : public FSlateWidgetStyle
 	FSlateBrush UnderlineBrush;
 	FTextBlockStyle& SetUnderlineBrush( const FSlateBrush& InUnderlineBrush ){ UnderlineBrush = InUnderlineBrush; return *this; }
 
+	/** The Text Transform Policy (defaults to None) */ 
+	UPROPERTY(EditAnywhere, BlueprintReadWrite, Category=Appearance, AdvancedDisplay)
+	ETextTransformPolicy TransformPolicy;
+	FTextBlockStyle& SetTransformPolicy( const ETextTransformPolicy& InTransformPolicy ){ TransformPolicy = InTransformPolicy; return *this; }
+
 	/**
 	 * Unlinks all colors in this style.
 	 * @see FSlateColor::Unlink
