Fix missing submenu indicators

diff --git a/Engine/Source/Editor/KismetWidgets/Private/SPinTypeSelector.cpp b/Engine/Source/Editor/KismetWidgets/Private/SPinTypeSelector.cpp
--- a/Engine/Source/Editor/KismetWidgets/Private/SPinTypeSelector.cpp
+++ b/Engine/Source/Editor/KismetWidgets/Private/SPinTypeSelector.cpp
@@ -626,7 +626,7 @@ TSharedRef<ITableRow> SPinTypeSelector::GenerateTypeTreeRow(FPinTypeTreeItem InI
 				.Padding(FMargin(7,0,0,0))
 				[
 					SNew( SImage )
-					.Image( FEditorStyle::Get().GetBrush( "ToolBar.SubMenuIndicator" ) )
+					.Image( FEditorStyle::Get().GetBrush( "Menu.SubMenuIndicator" ) )
 				]
 			];
 	}
diff --git a/Engine/Source/Runtime/AppFramework/Private/Framework/Testing/SLayoutExample.cpp b/Engine/Source/Runtime/AppFramework/Private/Framework/Testing/SLayoutExample.cpp
--- a/Engine/Source/Runtime/AppFramework/Private/Framework/Testing/SLayoutExample.cpp
+++ b/Engine/Source/Runtime/AppFramework/Private/Framework/Testing/SLayoutExample.cpp
@@ -130,7 +130,7 @@ public:
 				.VAlign(VAlign_Center)
 				[
 					SNew( SImage )
-					.Image( FCoreStyle::Get().GetBrush( "ToolBar.SubMenuIndicator" ) )
+					.Image( FCoreStyle::Get().GetBrush( "Menu.SubMenuIndicator" ) )
 				]
 			]
 			+SScrollBox::Slot() .Padding(5)
