Updated starship widget gallery with progress bar examples

diff --git a/Engine/Source/Runtime/AppFramework/Private/Framework/Testing/SStarshipGallery.cpp b/Engine/Source/Runtime/AppFramework/Private/Framework/Testing/SStarshipGallery.cpp
--- a/Engine/Source/Runtime/AppFramework/Private/Framework/Testing/SStarshipGallery.cpp
+++ b/Engine/Source/Runtime/AppFramework/Private/Framework/Testing/SStarshipGallery.cpp
@@ -584,7 +584,7 @@ public:
 
         TSharedPtr<SGridPanel> WidgetGrid = SNew(SGridPanel);
 
-        int WidgetNum = 0, RowCount = 11, Cols = 3;
+        int WidgetNum = 0, RowCount = 15, Cols = 3;
         auto NextSlot = [WidgetNum, RowCount, Cols](TSharedPtr<SGridPanel> Grid, const FText& InLabel) mutable -> SHorizontalBox::FSlot&
         {
 
@@ -664,52 +664,6 @@ public:
         	return SNew(SBox).HeightOverride(16.f)[ HBox.ToSharedRef() ];
         };
 
-        /*
-        // SBorder
-        NextSlot(WidgetGrid, LOCTEXT("SBorderLabel", "SBorder"))
-        [
-            SNew(SBorder)
-            [
-                SNew(SSpacer)
-                .Size(FVector2D(100.0f, 50.0f))
-            ]   
-        ];
-        :w
-        */
-
-
-        // SBorder
-        /*
-        NextSlot(WidgetGrid, LOCTEXT("RectTest", "FSlateRoundedBoxBrush(AccentColor, 24.0)"))// Test Rounded Box Brush"))
-        .Padding(20)
-        [
-            SNew(SBox).WidthOverride(200.0f).HeightOverride(48.0f)
-            [
-                SNew(SBorder).BorderImage(FAppStyle::Get().GetBrush( "RectTestBorder"))
-            ]
-        ];
-
-
-        NextSlot(WidgetGrid, LOCTEXT("RectTest2", "FSlateRoundedBoxBrush(FLinearColor::Transparent, 24.0f, FLinearColor::White, 2.0f)"))// Test Rounded Box Brush 2"))
-        .Padding(20)
-        [
-            SNew(SBox).WidthOverride(200.0f).HeightOverride(48.0f)
-            [
-                SNew(SImage).Image(FAppStyle::Get().GetBrush( "RectTestBorder2"))
-            ]
-        ];
-
-
-        NextSlot(WidgetGrid, LOCTEXT("RectTest3", "FSlateRoundedBoxBrush(FLinearColor::Transparent, 8.0f,  AccentColor, 8.0f)"))// Test Rounded Box Brush 3 "))
-        .Padding(20)
-        [
-            SNew(SBox).WidthOverride(200.0f).HeightOverride(48.0f)
-            [
-                SNew(SBorder).BorderImage(FAppStyle::Get().GetBrush( "RectTestBorder3"))
-            ]
-        ];
-    */
-
         // SButton Primary Rounded
         NextSlot(WidgetGrid, LOCTEXT("SButtonPrimaryExampleLabelRounded", "Primary Button"))
         [
@@ -1463,170 +1417,6 @@ public:
             .Value_Lambda( [this] { return TOptional<float>(NumericEntryBoxChoice); } )
         ];
 
-
-        /*
-        // STextBlock
-        NextSlot(WidgetGrid, LOCTEXT("STextBlockLabel", "STextBlock"))
-        [
-            SNew(SVerticalBox)
-        
-            + SVerticalBox::Slot()
-            .AutoHeight()
-            [
-                SNew(STextBlock)
-                    .Text(LOCTEXT("STextBlockLeftExampleLabel", "This is a left-justified text block"))
-                    .Justification(ETextJustify::Left)
-            ]
-
-            + SVerticalBox::Slot()
-            .AutoHeight()
-            [
-                SNew(STextBlock)
-                    .Text(LOCTEXT("STextBlockCenterExampleLabel", "This is a center-justified text block"))
-                    .Justification(ETextJustify::Center)
-            ]
-
-            + SVerticalBox::Slot()
-            .AutoHeight()
-            [
-                SNew(STextBlock)
-                    .Text(LOCTEXT("STextBlockRightExampleLabel", "This is a right-justified text block"))
-                    .Justification(ETextJustify::Right)
-            ]
-
-            + SVerticalBox::Slot()
-            .AutoHeight()
-            [
-                SNew(STextBlock)
-                    .TextStyle(FAppStyle::Get(), "NormalUnderlinedText")
-                    .Text(LOCTEXT("STextBlockUnderlineLabel", "This is text is underlined using a line highlighter!"))
-            ]
-        ];
-
-        // STextComboBox
-        NextSlot(WidgetGrid, LOCTEXT("STextComboBoxLabel", "STextComboBox"))
-        [
-            SAssignNew(TextComboBox, STextComboBox)
-                .OptionsSource(&TextComboBoxOptions)
-                .OnSelectionChanged(this, &SWidgetGallery::HandleTextComboBoxSelectionChanged)
-                .OnGetTextLabelForItem(this, &SWidgetGallery::HandleTextComboBoxGetTextLabelForItem)
-                .InitiallySelectedItem(TextComboBoxSelectedItem)            
-        ];
-
-        // STextComboPopup
-        NextSlot(WidgetGrid, LOCTEXT("STextComboPopupLabel", "STextComboPopup"))
-        [
-            SNew(SButton)
-                .Text(LOCTEXT("ButtonTestLabel", "Test"))
-                .OnClicked(this, &SWidgetGallery::HandleTextComboPopupClicked)
-        ];
-
-    */
-            /*
-        // SThrobber
-        WidgetGrid->AddSlot(0, 30)
-            [
-                SNew(STextBlock)
-                    .Text(LOCTEXT("SThrobberLabel", "SThrobber"))
-            ]
-
-        WidgetGrid->AddSlot(1, 30)
-            .Padding(0.0f, 5.0f)
-            [
-                SNew(SVerticalBox)
-
-                + SVerticalBox::Slot()
-                    .AutoHeight()
-                    [
-                        SNew(SThrobber)
-                            .Animate(SThrobber::Horizontal)
-                    ]
-
-                + SVerticalBox::Slot()
-                    .AutoHeight()
-                    [
-                        SNew(SThrobber)
-                            .Animate(SThrobber::Opacity)
-                    ]
-
-                + SVerticalBox::Slot()
-                    .AutoHeight()
-                    [
-                        SNew(SThrobber)
-                            .Animate(SThrobber::VerticalAndOpacity) .NumPieces(5)
-                    ]
-            ]
-        */
-/*
-        // SVectorInputBox
-        WidgetGrid->AddSlot(0, 31)
-            [
-                SNew(STextBlock)
-                    .Text(LOCTEXT("SVectorInputBoxLabel", "SVectorInputBox"))
-            ]
-
-        WidgetGrid->AddSlot(1, 31)
-            .Padding(0.0f, 5.0f)
-            [
-                SNew(SVectorInputBox)
-                    .X(-5)
-                    .Y(0)
-                    .Z(5)           
-            ]
-        // SVolumeControl
-        NextSlot(WidgetGrid, LOCTEXT("SVolumeControlLabel", "SVolumeControl"))
-        [
-            SNew(SBox)
-            .WidthOverride(150.0f)
-            [
-                SNew(SVolumeControl)
-                    .Volume(0.6f)
-            ]
-        ];
-
-        NextSlot(WidgetGrid, LOCTEXT("SUserWidgetLabel", "SUserWidgetExample"))
-        [
-            SNew(SUserWidgetExample)
-            .Title( LOCTEXT("SUserWidgetTest", "UserWidgetTest.cpp") )
-        ];
-
-        // Menu anchors that don't use the menu stack
-        NextSlot(WidgetGrid, LOCTEXT("StacklessAnchorLabel", "SMenuAnchor with UseApplicationMenuStack=false"))
-        [
-            SNew(SButton)
-            .OnClicked(this, &SWidgetGallery::OnClickStacklessAnchor, StacklessAnchor1)
-            [
-                SAssignNew(StacklessAnchor1, SMenuAnchor)
-                .Placement(MenuPlacement_BelowAnchor)
-                .Method(EPopupMethod::UseCurrentWindow)
-                .UseApplicationMenuStack(false)
-                .OnGetMenuContent(this, &SWidgetGallery::MakeAnchorMenuContent, LOCTEXT("StacklessContentLabel1", "Stackless Content 1"))
-                [
-                    SNew(STextBlock)
-                    .Text(LOCTEXT("ClickMeLabel", "Click Me!"))
-                ]
-            ]
-        ];
-
-        // Menu anchors that don't use the menu stack
-        NextSlot(WidgetGrid, LOCTEXT("StacklessAnchorLabel", "SMenuAnchor with UseApplicationMenuStack=false"))
-        [
-            SNew(SButton)
-            .OnClicked(this, &SWidgetGallery::OnClickStacklessAnchor, StacklessAnchor2)
-            [
-                SAssignNew(StacklessAnchor2, SMenuAnchor)
-                .Placement(MenuPlacement_BelowAnchor)
-                .Method(EPopupMethod::CreateNewWindow)
-                .UseApplicationMenuStack(false)
-                .OnGetMenuContent(this, &SWidgetGallery::MakeAnchorMenuContent, LOCTEXT("StacklessContentLabel2", "Stackless Content 2"))
-                [
-                    SNew(STextBlock)
-                    .Text(LOCTEXT("ClickMeLabel", "Click Me!"))
-                ]
-            ]
-        ];
-        */
-
         // SEditableText
         NextSlot(WidgetGrid, LOCTEXT("SEditableTextLabel", "SEditableText"))
         .FillWidth(1.0)
@@ -1689,6 +1479,20 @@ public:
             SNew(SSearchBox)
         ];
 
+		// SProgressBar
+		NextSlot(WidgetGrid, LOCTEXT("SSProgressBar", "SProgressBar"))
+			.FillWidth(1.0)
+			[
+				SNew(SProgressBar)
+				.Percent(.5f)
+			];
+
+		// SProgressBar marquee
+		NextSlot(WidgetGrid, LOCTEXT("SSProgressBarMarquee", "SProgressBar marquee"))
+			.FillWidth(1.0)
+			[
+				SNew(SProgressBar)
+			];
 
         // Window Flash Button
         NextSlot(WidgetGrid, LOCTEXT("WindowFlashAction", "Window Flash Button"))
diff --git a/Engine/Source/Runtime/AppFramework/Private/Widgets/Testing/SStarshipSuite.cpp b/Engine/Source/Runtime/AppFramework/Private/Widgets/Testing/SStarshipSuite.cpp
--- a/Engine/Source/Runtime/AppFramework/Private/Widgets/Testing/SStarshipSuite.cpp
+++ b/Engine/Source/Runtime/AppFramework/Private/Widgets/Testing/SStarshipSuite.cpp
@@ -29,7 +29,6 @@ void RestoreStarshipSuite()
 	{
 		return SNew(SDockTab)
 			.TabRole(ETabRole::NomadTab)
-			.Icon(new FSlateDynamicImageBrush(FName("../../Content/Slate/Starship/Icons/Quixel/Common/gift_32x.png"), FVector2D(16.f, 16.f)))
 			[
 				MakeStarshipGallery()
 			];
