Removes UE4 references in AppFramework

diff --git a/Engine/Source/Runtime/AppFramework/Private/Framework/Testing/TestStyle.cpp b/Engine/Source/Runtime/AppFramework/Private/Framework/Testing/TestStyle.cpp
--- a/Engine/Source/Runtime/AppFramework/Private/Framework/Testing/TestStyle.cpp
+++ b/Engine/Source/Runtime/AppFramework/Private/Framework/Testing/TestStyle.cpp
@@ -41,6 +41,7 @@ void FTestStyle::SetStyle( const TSharedRef< ISlateStyle >& NewStyle )
 }
 
 #define IMAGE_BRUSH( RelativePath, ... ) FSlateImageBrush( Style->RootToContentDir( RelativePath, TEXT(".png") ), __VA_ARGS__ )
+#define IMAGE_BRUSH_SVG( RelativePath, ... ) FSlateVectorImageBrush( Style->RootToContentDir(RelativePath, TEXT(".svg")), __VA_ARGS__)
 #define BOX_BRUSH( RelativePath, ... ) FSlateBoxBrush( Style->RootToContentDir( RelativePath, TEXT(".png") ), __VA_ARGS__ )
 #define BORDER_BRUSH( RelativePath, ... ) FSlateBorderBrush( Style->RootToContentDir( RelativePath, TEXT(".png") ), __VA_ARGS__ )
 #define DEFAULT_FONT(...) FCoreStyle::GetDefaultFontStyle(__VA_ARGS__)
@@ -85,7 +86,7 @@ TSharedRef< ISlateStyle > FTestStyle::Create()
 
 	Style->Set( "GammaReference", new IMAGE_BRUSH( "Testing/GammaReference", FVector2D(256,128) ) );
 
-	Style->Set( "UE4Icon", new IMAGE_BRUSH( "Testing/UE4Icon", FVector2D(50,50) ) );
+	Style->Set( "UEIcon", new IMAGE_BRUSH_SVG( "Starship/Common/unreal", FVector2D(50,50) ) );
 
 	Style->Set( "NewLevelBlank", new IMAGE_BRUSH( "Testing/NewLevelBlank", FVector2D(256,256) ) );
 
diff --git a/Engine/Source/Runtime/AppFramework/Private/Widgets/Testing/STestSuite.cpp b/Engine/Source/Runtime/AppFramework/Private/Widgets/Testing/STestSuite.cpp
--- a/Engine/Source/Runtime/AppFramework/Private/Widgets/Testing/STestSuite.cpp
+++ b/Engine/Source/Runtime/AppFramework/Private/Widgets/Testing/STestSuite.cpp
@@ -4251,7 +4251,7 @@ class SFxTest : public SCompoundWidget
 							+SHorizontalBox::Slot() .AutoWidth()
 							[
 								SNew(SImage)
-								.Image(FTestStyle::Get().GetBrush("UE4Icon"))
+								.Image(FTestStyle::Get().GetBrush("UEIcon"))
 							]
 							+SHorizontalBox::Slot() .AutoWidth()
 							[
@@ -4726,7 +4726,7 @@ public:
 					.AutoWidth()
 					[
 						SNew(SImage)
-						.Image(FTestStyle::Get().GetBrush("UE4Icon"))
+						.Image(FTestStyle::Get().GetBrush("UEIcon"))
 					]
 					+ SHorizontalBox::Slot()
 					.Padding(FMargin(15.0f, 0.0f, 0.0f, 0.0f))
@@ -5585,7 +5585,7 @@ private:
 		{
 			check(IsInGameThread());
 			Style.SetContentRoot(FPaths::EngineContentDir() / TEXT("Slate"));
-			Style.Set("UE4Icon", new FSlateImageBrush(Style.RootToContentDir(TEXT("Testing/UE4Icon.png")), FVector2D(50, 50)));
+			Style.Set("UEIcon", new FSlateVectorImageBrush(Style.RootToContentDir(TEXT("Starship/Common/unreal.svg")), FVector2D(50, 50)));
 			IsInit = true;
 		}
 		return Style;
@@ -5615,7 +5615,7 @@ void SRenderTransformManipulatorWidgetImpl::Construct(const FArguments& InArgs)
 					SAssignNew(ImageWidget, SImage)
 					.RenderTransform_Static(&::GetTestRenderTransform)
 					.RenderTransformPivot_Static(&::GetTestRenderTransformPivot)
-					.Image(GetStyle().GetBrush("UE4Icon"))
+					.Image(GetStyle().GetBrush("UEIcon"))
 				]
 			+ SVerticalBox::Slot().AutoHeight().Padding(4)
 			[
