PR #8196: Fix the settings of SColorPicker (Contributed by [REDACTED_PERSON])

diff --git a/Engine/Source/Runtime/AppFramework/Private/Widgets/Colors/SColorPicker.cpp b/Engine/Source/Runtime/AppFramework/Private/Widgets/Colors/SColorPicker.cpp
--- a/Engine/Source/Runtime/AppFramework/Private/Widgets/Colors/SColorPicker.cpp
+++ b/Engine/Source/Runtime/AppFramework/Private/Widgets/Colors/SColorPicker.cpp
@@ -106,7 +106,7 @@ void SColorPicker::Construct( const FArguments& InArgs )
 
 	bool bAdvancedSectionExpanded = false;
 
-	if (!FPaths::FileExists(GEditorPerProjectIni))
+	if (FPaths::FileExists(GEditorPerProjectIni))
 	{
 		bool WheelMode = true;
 
