Made ColorPicker buttons pop out and fixed an issue that caused the picker to resize

diff --git a/Engine/Source/Editor/EditorStyle/Private/StarshipStyle.cpp b/Engine/Source/Editor/EditorStyle/Private/StarshipStyle.cpp
--- a/Engine/Source/Editor/EditorStyle/Private/StarshipStyle.cpp
+++ b/Engine/Source/Editor/EditorStyle/Private/StarshipStyle.cpp
@@ -287,6 +287,7 @@ void FStarshipEditorStyle::FStyle::Initialize()
 	SetupUMGEditorStyles();
 	SetupMyBlueprintStyles();
 	SetupStatusBarStyle();
+	SetupColorPickerStyle();
 
 //	LogUnusedBrushResources();
 
@@ -7787,6 +7788,11 @@ void FStarshipEditorStyle::FStyle::SetupStatusBarStyle()
 		.SetColorAndOpacity(FStyleColors::Primary));
 }
 
+void FStarshipEditorStyle::FStyle::SetupColorPickerStyle()
+{
+	Set("ColorPicker.ColorThemes", new IMAGE_BRUSH_SVG("Starship/ColorPicker/ColorThemes", Icon16x16));
+}
+
 #undef DEFAULT_FONT
 #undef ICON_FONT
 #undef LOCTEXT_NAMESPACE
diff --git a/Engine/Source/Editor/EditorStyle/Private/StarshipStyle.h b/Engine/Source/Editor/EditorStyle/Private/StarshipStyle.h
--- a/Engine/Source/Editor/EditorStyle/Private/StarshipStyle.h
+++ b/Engine/Source/Editor/EditorStyle/Private/StarshipStyle.h
@@ -71,6 +71,7 @@ public:
 		void SetupUMGEditorStyles();
 		void SetupMyBlueprintStyles();
 		void SetupStatusBarStyle();
+		void SetupColorPickerStyle();
 
 		void SettingsChanged(UObject* ChangedObject, FPropertyChangedEvent& PropertyChangedEvent);
 		void SyncSettings();
diff --git a/Engine/Source/Runtime/AppFramework/Private/Widgets/Colors/SColorPicker.cpp b/Engine/Source/Runtime/AppFramework/Private/Widgets/Colors/SColorPicker.cpp
--- a/Engine/Source/Runtime/AppFramework/Private/Widgets/Colors/SColorPicker.cpp
+++ b/Engine/Source/Runtime/AppFramework/Private/Widgets/Colors/SColorPicker.cpp
@@ -241,7 +241,7 @@ BEGIN_SLATE_FUNCTION_BUILD_OPTIMIZATION
 void SColorPicker::GenerateDefaultColorPickerContent( bool bAdvancedSectionExpanded )
 {	
 	// The height of the gradient bars beneath the sliders
-	const FSlateFontInfo SmallLayoutFont = FCoreStyle::Get().GetFontStyle("ColorPicker.Font");
+	const FSlateFontInfo SmallLayoutFont = FAppStyle::Get().GetFontStyle("ColorPicker.Font");
 
 	TSharedPtr<SColorThemesViewer> ThemesViewer = ColorThemesViewer.Pin();
 
@@ -257,9 +257,13 @@ void SColorPicker::GenerateDefaultColorPickerContent( bool bAdvancedSectionExpan
 
 	// The standard button to open the color themes can temporarily become a trash for colors
 	ColorThemeComboButton = SNew(SComboButton)
-		.ContentPadding(3.0f)
-		.MenuPlacement(MenuPlacement_ComboBox)
-		.ToolTipText(LOCTEXT("OpenThemeManagerToolTip", "Open Color Theme Manager"));
+		.ComboButtonStyle(FAppStyle::Get(), "SimpleComboButton")
+		.VAlign(VAlign_Center)
+		.ButtonContent()
+		[
+			SNew(SImage)
+			.Image(FAppStyle::Get().GetBrush("ColorPicker.ColorThemes"))
+		];
 
 	ColorThemeComboButton->SetMenuContent(ThemesViewer.ToSharedRef());
 
@@ -285,9 +289,9 @@ void SColorPicker::GenerateDefaultColorPickerContent( bool bAdvancedSectionExpan
 					.Padding(0.0f, 1.0f, 20.0f, 1.0f)
 					[
 						SNew(SHorizontalBox)
-
+						
 						+ SHorizontalBox::Slot()
-							.FillWidth(1.0f)
+							.FillWidth(0.91f)
 							.Padding(0.0f, 1.0f)
 							[
 								SNew(SOverlay)
@@ -318,7 +322,7 @@ void SColorPicker::GenerateDefaultColorPickerContent( bool bAdvancedSectionExpan
 							]
 
 						+ SHorizontalBox::Slot()
-							.AutoWidth()
+							.FillWidth(0.09f)
 							[
 								// color theme selector
 								SAssignNew(ColorThemeButtonOrSmallTrash, SBorder)
@@ -346,7 +350,7 @@ void SColorPicker::GenerateDefaultColorPickerContent( bool bAdvancedSectionExpan
 					.Padding(0.0f, 8.0f, 20.0f, 0.0f)
 					[
 						SNew(SBorder)
-							.BorderImage(FCoreStyle::Get().GetBrush("NoBorder"))
+							.BorderImage(FAppStyle::Get().GetBrush("NoBorder"))
 							.Padding(0.0f)
 							.OnMouseButtonDown(this, &SColorPicker::HandleColorAreaMouseDown)
 							[
@@ -436,7 +440,7 @@ void SColorPicker::GenerateDefaultColorPickerContent( bool bAdvancedSectionExpan
 											.Content()
 											[
 												SNew(SImage)
-													.Image(FCoreStyle::Get().GetBrush("ColorPicker.Mode"))
+													.Image(FAppStyle::Get().GetBrush("ColorPicker.Mode"))
 													.ToolTipText(LOCTEXT("ColorPickerModeEToolTip", "Toggle between color wheel and color spectrum."))
 											]									
 									]
@@ -600,16 +604,16 @@ void SColorPicker::GenerateDefaultColorPickerContent( bool bAdvancedSectionExpan
 			.Padding(0.0f, 12.0f, 0.0f, 0.0f)
 			[
 				SNew(SUniformGridPanel)
-					.MinDesiredSlotHeight(FCoreStyle::Get().GetFloat("StandardDialog.MinDesiredSlotHeight"))
-					.MinDesiredSlotWidth(FCoreStyle::Get().GetFloat("StandardDialog.MinDesiredSlotWidth"))
-					.SlotPadding(FCoreStyle::Get().GetMargin("StandardDialog.SlotPadding"))
+					.MinDesiredSlotHeight(FAppStyle::Get().GetFloat("StandardDialog.MinDesiredSlotHeight"))
+					.MinDesiredSlotWidth(FAppStyle::Get().GetFloat("StandardDialog.MinDesiredSlotWidth"))
+					.SlotPadding(FAppStyle::Get().GetMargin("StandardDialog.SlotPadding"))
 					.Visibility((ParentWindowPtr.IsValid() || bValidCreationOverrideExists) ? EVisibility::Visible : EVisibility::Collapsed)
 
 				+ SUniformGridPanel::Slot(0, 0)
 					[
 						// ok button
 						SNew(SButton)
-							.ContentPadding( FCoreStyle::Get().GetMargin("StandardDialog.ContentPadding") )
+							.ContentPadding(FAppStyle::Get().GetMargin("StandardDialog.ContentPadding") )
 							.HAlign(HAlign_Center)
 							.Text(LOCTEXT("OKButton", "OK"))
 							.OnClicked(this, &SColorPicker::HandleOkButtonClicked)
@@ -619,7 +623,7 @@ void SColorPicker::GenerateDefaultColorPickerContent( bool bAdvancedSectionExpan
 					[
 						// cancel button
 						SNew(SButton)
-							.ContentPadding( FCoreStyle::Get().GetMargin("StandardDialog.ContentPadding") )
+							.ContentPadding(FAppStyle::Get().GetMargin("StandardDialog.ContentPadding") )
 							.HAlign(HAlign_Center)
 							.Text(LOCTEXT("CancelButton", "Cancel"))
 							.OnClicked(this, &SColorPicker::HandleCancelButtonClicked)
@@ -890,7 +894,7 @@ TSharedRef<SWidget> SColorPicker::MakeColorSlider( EColorPickerChannels Channel
 				.IndentHandle(false)
 				.Orientation(Orient_Vertical)
 				.SliderBarColor(FLinearColor::Transparent)
-				.Style(&FCoreStyle::Get().GetWidgetStyle<FSliderStyle>("ColorPicker.Slider"))
+				.Style(&FAppStyle::Get().GetWidgetStyle<FSliderStyle>("ColorPicker.Slider"))
 				.Value(this, &SColorPicker::HandleColorSpinBoxValue, Channel)
 				.OnMouseCaptureBegin(const_cast<SColorPicker*>(this), &SColorPicker::HandleInteractiveChangeBegin)
 				.OnMouseCaptureEnd(const_cast<SColorPicker*>(this), &SColorPicker::HandleInteractiveChangeEnd)
@@ -908,7 +912,7 @@ TSharedRef<SWidget> SColorPicker::MakeColorSpinBox( EColorPickerChannels Channel
 
 	const int32 GradientHeight = 6;
 	const float HDRMaxValue = (TargetFColors.Num()) ? 1.f : FLT_MAX;
-	const FSlateFontInfo SmallLayoutFont = FCoreStyle::Get().GetFontStyle("ColorPicker.Font");
+	const FSlateFontInfo SmallLayoutFont = FAppStyle::Get().GetFontStyle("ColorPicker.Font");
 
 	// create gradient widget
 	TSharedPtr<SWidget> GradientWidget;
@@ -1140,7 +1144,7 @@ TSharedRef<SWidget> SColorPicker::MakeColorPreviewBox() const
 						.WidthOverride(4.0f)
 						[
 							SNew(SBorder)
-								.BorderImage(FCoreStyle::Get().GetBrush("ColorPicker.Separator"))
+								.BorderImage(FAppStyle::Get().GetBrush("ColorPicker.Separator"))
 								.Padding(0.0f)
 						]
 				]
@@ -1154,7 +1158,7 @@ TSharedRef<SWidget> SColorPicker::MakeColorPreviewBox() const
 						.WidthOverride(4.0f)
 						[
 							SNew(SBorder)
-								.BorderImage(FCoreStyle::Get().GetBrush("ColorPicker.Separator"))
+								.BorderImage(FAppStyle::Get().GetBrush("ColorPicker.Separator"))
 								.Padding(0.0f)
 						]
 				]
@@ -1684,7 +1688,7 @@ bool OpenColorPicker(const FColorPickerArgs& Args)
 
 	TSharedPtr<SWindow> Window = nullptr;
 	TSharedRef<SBorder> WindowContent = SNew(SBorder)
-			.BorderImage(FCoreStyle::Get().GetBrush("ToolPanel.GroupBorder"))
+			.BorderImage(FAppStyle::Get().GetBrush("Brushes.Panel"))
 			.Padding(FMargin(8.0f, 8.0f));
 	
 	bool bNeedToAddWindow = true;
diff --git a/Engine/Source/Runtime/AppFramework/Private/Widgets/Colors/SColorThemes.cpp b/Engine/Source/Runtime/AppFramework/Private/Widgets/Colors/SColorThemes.cpp
--- a/Engine/Source/Runtime/AppFramework/Private/Widgets/Colors/SColorThemes.cpp
+++ b/Engine/Source/Runtime/AppFramework/Private/Widgets/Colors/SColorThemes.cpp
@@ -214,7 +214,7 @@ FReply SColorTrash::OnDrop( const FGeometry& MyGeometry, const FDragDropEvent& D
 const FSlateBrush* SColorTrash::GetBorderStyle() const
 {
 	return (bBorderActivated) ?
-		FCoreStyle::Get().GetBrush("FocusRectangle") :
+		FAppStyle::Get().GetBrush("FocusRectangle") :
 		FStyleDefaults::GetNoBrush();
 }
 
@@ -238,8 +238,8 @@ void SThemeColorBlock::Construct(const FArguments& InArgs )
 
 	DistanceDragged = 0;
 
-	const FSlateFontInfo SmallLayoutFont = FCoreStyle::GetDefaultFontStyle("Regular", 9);
-	const FSlateFontInfo SmallLabelFont = FCoreStyle::GetDefaultFontStyle("Bold", 9);
+	const FSlateFontInfo SmallLayoutFont = FAppStyle::Get().GetFontStyle("Regular");
+	const FSlateFontInfo SmallLabelFont = FAppStyle::Get().GetFontStyle("Bold");
 
 	TSharedPtr<SToolTip> ColorTooltip =
 		SNew(SToolTip)
@@ -436,8 +436,8 @@ FSlateColor SThemeColorBlock::HandleBorderColor() const
 const FSlateBrush* SThemeColorBlock::HandleBorderImage() const
 {
 	return IsHovered() ?
-		FCoreStyle::Get().GetBrush("FocusRectangle") :
-		FCoreStyle::Get().GetBrush("GenericWhiteBox");
+		FAppStyle::Get().GetBrush("FocusRectangle") :
+		FAppStyle::Get().GetBrush("GenericWhiteBox");
 }
 
 #define LOCTEXT_NAMESPACE "ThemeColorBlock"
@@ -830,28 +830,29 @@ void SColorThemesViewer::Construct(const FArguments& InArgs)
 
 	LoadColorThemesFromIni();
 
-	const FSlateFontInfo SmallLayoutFont = FCoreStyle::GetDefaultFontStyle("Regular", 10);
+	const FSlateFontInfo SmallLayoutFont = FAppStyle::Get().GetFontStyle("Regular");
+
 
 	// different menus that could be visible for the color themes menu
 	// standard menu with "new", "rename" and "delete"
 	MenuStandard =
 		SNew(SHorizontalBox)
-		+SHorizontalBox::Slot() .FillWidth(1) .Padding(3)
+		+SHorizontalBox::Slot() .AutoWidth().Padding(3)
 		[
 			SNew(SButton) .Text(LOCTEXT("NewButton", "New")) .HAlign(HAlign_Center)
 				.OnClicked(this, &SColorThemesViewer::NewColorTheme)
 		]
-		+SHorizontalBox::Slot() .FillWidth(1) .Padding(3)
+		+SHorizontalBox::Slot() .AutoWidth().Padding(3)
 		[
 			SNew(SButton) .Text(LOCTEXT("DuplicateButton", "Duplicate")) .HAlign(HAlign_Center)
 				.OnClicked(this, &SColorThemesViewer::DuplicateColorTheme)
 		]
-		+SHorizontalBox::Slot() .FillWidth(1) .Padding(3)
+		+SHorizontalBox::Slot() .AutoWidth().Padding(3)
 		[
 			SNew(SButton) .Text(LOCTEXT("RenameButton", "Rename")) .HAlign(HAlign_Center)
 				.OnClicked(this, &SColorThemesViewer::MenuToRename)
 		]
-		+SHorizontalBox::Slot() .FillWidth(1) .Padding(3)
+		+SHorizontalBox::Slot() .AutoWidth().Padding(3)
 		[
 			SNew(SButton) .Text(LOCTEXT("DeleteButton", "Delete")) .HAlign(HAlign_Center)
 				.OnClicked(this, &SColorThemesViewer::MenuToDelete)
@@ -925,6 +926,7 @@ void SColorThemesViewer::Construct(const FArguments& InArgs)
 			.Padding(0,15,0,0)
 		[
 			SAssignNew(Menu, SBorder)
+			.BorderImage(FAppStyle::Get().GetBrush("Brushes.Panel"))
 		]
 		+ SVerticalBox::Slot()
 			.AutoHeight()
