Re-enable erroneously disabled tests in STestSuite. Hide "global" scope variables for MSVC.

diff --git a/Engine/Source/Runtime/AppFramework/Private/Widgets/Testing/STestSuite.cpp b/Engine/Source/Runtime/AppFramework/Private/Widgets/Testing/STestSuite.cpp
--- a/Engine/Source/Runtime/AppFramework/Private/Widgets/Testing/STestSuite.cpp
+++ b/Engine/Source/Runtime/AppFramework/Private/Widgets/Testing/STestSuite.cpp
@@ -153,8 +153,6 @@ namespace
 		FMatrix mat5;
 		FQuat quat5;
 
-		// TODO: LWC: these cause compiler errors about imminent stack overflow in the Concatenate() function for FMatrix.
-		/*
 		mat5 = Concatenate(uniScale4, trans4);
 		mat5 = Concatenate(trans4, uniScale4);
 		mat5 = Concatenate(uniScale4, quat4);
@@ -182,7 +180,6 @@ namespace
 		mat5 = Concatenate(Inverse(trans), rot, trans);
 
 		mat5 = Concatenate(Inverse(trans), rot, uniScale, trans);
-		*/
 
 		FScale2D scale2D(1.4f, 3.2f);
 		FShear2D shear2D(-3.5f, -4.6f);
@@ -5563,7 +5560,7 @@ public:
 	virtual void Construct(const FArguments& InArgs) = 0;
 };
 
-namespace
+namespace TestSuiteNS
 {
 	float RotDeg = 0.0f;
 	FScale2D Scale;
@@ -5599,12 +5596,12 @@ private:
  */
 TOptional<FSlateRenderTransform> GetTestRenderTransform()
 {
-	return TransformCast<FSlateRenderTransform>(Concatenate(Shear, Scale, Rot));
+	return TransformCast<FSlateRenderTransform>(Concatenate(TestSuiteNS::Shear, TestSuiteNS::Scale, TestSuiteNS::Rot));
 }
 
 FVector2D GetTestRenderTransformPivot()
 {
-	return Offset;
+	return TestSuiteNS::Offset;
 }
 
 void SRenderTransformManipulatorWidgetImpl::Construct(const FArguments& InArgs)
@@ -5631,8 +5628,8 @@ void SRenderTransformManipulatorWidgetImpl::Construct(const FArguments& InArgs)
 						SNew(SSpinBox<float>)
 						.MinValue(0.0f)
 						.MaxValue(360.0f)
-						.OnValueChanged_Static([](float val) { RotDeg = val; Rot = FQuat2D(FMath::DegreesToRadians(val)); })
-						.Value_Static([] { return RotDeg; })
+						.OnValueChanged_Static([](float val) { TestSuiteNS::RotDeg = val; TestSuiteNS::Rot = FQuat2D(FMath::DegreesToRadians(val)); })
+						.Value_Static([] { return TestSuiteNS::RotDeg; })
 					]
 			]
 			+ SVerticalBox::Slot().AutoHeight().Padding(4)
@@ -5647,8 +5644,8 @@ void SRenderTransformManipulatorWidgetImpl::Construct(const FArguments& InArgs)
 							SNew(SSpinBox<float>)
 							.MinValue(-2.0f)
 							.MaxValue(2.0f)
-							.OnValueChanged_Static([](float val) { Offset.X = val; })
-							.Value_Static([] { return Offset.X; })
+							.OnValueChanged_Static([](float val) { TestSuiteNS::Offset.X = val; })
+							.Value_Static([] { return TestSuiteNS::Offset.X; })
 						]
 				]
 			+ SVerticalBox::Slot().AutoHeight().Padding(4)
@@ -5663,8 +5660,8 @@ void SRenderTransformManipulatorWidgetImpl::Construct(const FArguments& InArgs)
 							SNew(SSpinBox<float>)
 							.MinValue(-2.0f)
 							.MaxValue(2.0f)
-							.OnValueChanged_Static([](float val) { Offset.Y = val; })
-							.Value_Static([] { return Offset.Y; })
+							.OnValueChanged_Static([](float val) { TestSuiteNS::Offset.Y = val; })
+							.Value_Static([] { return TestSuiteNS:: Offset.Y; })
 						]
 				]
 			+ SVerticalBox::Slot().AutoHeight().Padding(4)
@@ -5679,8 +5676,8 @@ void SRenderTransformManipulatorWidgetImpl::Construct(const FArguments& InArgs)
 							SNew(SSpinBox<float>)
 							.MinValue(-10.0f)
 							.MaxValue(10.0f)
-							.OnValueChanged_Static([](float val) { Scale = FScale2D(val, Scale.GetVector().Y); })
-							.Value_Static([] { return Scale.GetVector().X; })
+							.OnValueChanged_Static([](float val) { TestSuiteNS::Scale = FScale2D(val, TestSuiteNS::Scale.GetVector().Y); })
+							.Value_Static([] { return TestSuiteNS::Scale.GetVector().X; })
 						]
 				]
 			+ SVerticalBox::Slot().AutoHeight().Padding(4)
@@ -5695,8 +5692,8 @@ void SRenderTransformManipulatorWidgetImpl::Construct(const FArguments& InArgs)
 							SNew(SSpinBox<float>)
 							.MinValue(-10.0f)
 							.MaxValue(10.0f)
-							.OnValueChanged_Static([](float val) { Scale = FScale2D(Scale.GetVector().X, val); })
-							.Value_Static([] { return Scale.GetVector().Y; })
+							.OnValueChanged_Static([](float val) { TestSuiteNS::Scale = FScale2D(TestSuiteNS::Scale.GetVector().X, val); })
+							.Value_Static([] { return TestSuiteNS::Scale.GetVector().Y; })
 						]
 				]
 			+ SVerticalBox::Slot().AutoHeight().Padding(4)
@@ -5711,8 +5708,8 @@ void SRenderTransformManipulatorWidgetImpl::Construct(const FArguments& InArgs)
 							SNew(SSpinBox<float>)
 							.MinValue(-4.f)
 							.MaxValue(4.f)
-							.OnValueChanged_Static([](float val) { Shear = FShear2D(val, Shear.GetVector().Y); })
-							.Value_Static([] { return Shear.GetVector().X; })
+							.OnValueChanged_Static([](float val) { TestSuiteNS::Shear = FShear2D(val, TestSuiteNS::Shear.GetVector().Y); })
+							.Value_Static([] { return TestSuiteNS::Shear.GetVector().X; })
 						]
 				]
 			+ SVerticalBox::Slot().AutoHeight().Padding(4)
@@ -5727,8 +5724,8 @@ void SRenderTransformManipulatorWidgetImpl::Construct(const FArguments& InArgs)
 							SNew(SSpinBox<float>)
 							.MinValue(-4.f)
 							.MaxValue(4.f)
-							.OnValueChanged_Static([](float val) { Shear = FShear2D(Shear.GetVector().X, val); })
-							.Value_Static([] { return Shear.GetVector().Y; })
+							.OnValueChanged_Static([](float val) { TestSuiteNS::Shear = FShear2D(TestSuiteNS::Shear.GetVector().X, val); })
+							.Value_Static([] { return TestSuiteNS::Shear.GetVector().Y; })
 						]
 				]
 		]
