Add graph editor icons section to starship icons gallery

diff --git a/Engine/Source/Runtime/AppFramework/Private/Framework/Testing/SStarshipGallery.cpp b/Engine/Source/Runtime/AppFramework/Private/Framework/Testing/SStarshipGallery.cpp
--- a/Engine/Source/Runtime/AppFramework/Private/Framework/Testing/SStarshipGallery.cpp
+++ b/Engine/Source/Runtime/AppFramework/Private/Framework/Testing/SStarshipGallery.cpp
@@ -457,6 +457,7 @@ public:
                 +SVerticalBox::Slot().AutoHeight()[ GenerateIconLibrarySVG(NSLOCTEXT("StarshipGallery", "Editor Common", "Editor"), "Content/Editor/Slate/Starship/Common")]
 
                 +SVerticalBox::Slot().AutoHeight()[ GenerateIconLibrarySVG(NSLOCTEXT("StarshipGallery", "SceneOutliner", "SceneOutliner"), "Content/Editor/Slate/Starship/SceneOutliner")]
+                +SVerticalBox::Slot().AutoHeight()[ GenerateIconLibrarySVG(NSLOCTEXT("StarshipGallery", "GraphEditors", "GraphEditors"), "Content/Editor/Slate/Starship/GraphEditors")]
                 // +SVerticalBox::Slot().AutoHeight()[ GenerateIconLibrarySVG(NSLOCTEXT("StarshipGallery", "LevelEditor", "LevelEditor"), "Content/Editor/Slate/Starship/LevelEditor/Menus")]
                 +SVerticalBox::Slot().AutoHeight()[ GenerateIconLibrarySVG(NSLOCTEXT("StarshipGallery", "MainToolbar", "MainToolbar"), "Content/Editor/Slate/Starship/MainToolbar")]
                 // +SVerticalBox::Slot().AutoHeight()[ GenerateIconLibrarySVG(NSLOCTEXT("StarshipGallery", "FileMenu", "FileMenu"), "Content/Editor/Slate/Starship/Menus/File")]
