Update starship gallery with new styles

diff --git a/Engine/Source/Runtime/AppFramework/Private/Framework/Testing/SStarshipGallery.cpp b/Engine/Source/Runtime/AppFramework/Private/Framework/Testing/SStarshipGallery.cpp
--- a/Engine/Source/Runtime/AppFramework/Private/Framework/Testing/SStarshipGallery.cpp
+++ b/Engine/Source/Runtime/AppFramework/Private/Framework/Testing/SStarshipGallery.cpp
@@ -61,6 +61,7 @@
 
 
 #include "Brushes/SlateImageBrush.h"
+#include "Framework/MultiBox/MultiBoxBuilder.h"
 
 #if !UE_BUILD_SHIPPING
 
@@ -560,6 +561,91 @@ public:
 
     float NumericEntryBoxChoice;
 
+	void MakeMenuEntries(FMenuBuilder& MenuBuilder)
+	{
+		MenuBuilder.AddMenuEntry(LOCTEXT("MenuEntry1", "Menu Entry 1"), FText::GetEmpty(), FSlateIcon(FAppStyle::GetAppStyleSetName(), "Icons.box-perspective"), FUIAction());
+		MenuBuilder.AddMenuEntry(LOCTEXT("MenuEntry2", "Menu Entry 2"), FText::GetEmpty(), FSlateIcon(FAppStyle::GetAppStyleSetName(), "Icons.box-perspective"), FUIAction());
+		MenuBuilder.AddMenuEntry(LOCTEXT("MenuEntry3", "Menu Entry 3"), FText::GetEmpty(), FSlateIcon(FAppStyle::GetAppStyleSetName(), "Icons.box-perspective"), FUIAction());
+		MenuBuilder.AddMenuEntry(LOCTEXT("MenuEntry4", "Menu Entry 4"), FText::GetEmpty(), FSlateIcon(FAppStyle::GetAppStyleSetName(), "Icons.box-perspective"), FUIAction());
+	};
+
+	TSharedRef<SWidget> BuildToolbar(bool bIncludeSettingsButtons)
+	{
+		FSlimHorizontalToolBarBuilder Toolbar(nullptr, FMultiBoxCustomization::None);
+
+		FUIAction ButtonAction;
+		Toolbar.AddToolBarButton(ButtonAction, NAME_None, LOCTEXT("ToolBarButton1", "Button"), FText::GetEmpty(), FSlateIcon(FAppStyle::GetAppStyleSetName(), "Icons.box-perspective"));
+
+		static bool IsChecked = false;
+		FUIAction ToggleButtonAction;
+		ToggleButtonAction.GetActionCheckState.BindLambda([&] {return IsChecked ? ECheckBoxState::Checked : ECheckBoxState::Unchecked; });
+		ToggleButtonAction.ExecuteAction.BindLambda([&] {IsChecked = !IsChecked; });
+
+		Toolbar.AddToolBarButton(ToggleButtonAction, NAME_None, LOCTEXT("ToolBarButtonToggle", "Toggle"), FText::GetEmpty(), FSlateIcon(FAppStyle::GetAppStyleSetName(), "Icons.box-perspective"), EUserInterfaceActionType::ToggleButton);
+	
+		FUIAction ComboAction;
+		Toolbar.AddComboButton(ComboAction, FOnGetContent::CreateLambda(
+			[&]()
+			{
+				FMenuBuilder Menu(true, nullptr);
+				MakeMenuEntries(Menu);
+
+				return Menu.MakeWidget();
+			}),
+			LOCTEXT("Dropdown","Dropdown"),
+			FText::GetEmpty(),
+			FSlateIcon(FAppStyle::GetAppStyleSetName(), "Icons.box-perspective")
+		);
+
+		if(bIncludeSettingsButtons)
+		{
+			Toolbar.BeginSection("SettingsTest");
+			{
+				Toolbar.AddToolBarButton(ButtonAction, NAME_None, LOCTEXT("ToolBarButton1", "Button"), FText::GetEmpty(), FSlateIcon(FAppStyle::GetAppStyleSetName(), "Icons.box-perspective"));
+
+				Toolbar.AddComboButton(
+					FUIAction(),
+					FOnGetContent::CreateLambda(
+						[&]()
+						{
+							FMenuBuilder Menu(true, nullptr);
+							MakeMenuEntries(Menu);
+
+							return Menu.MakeWidget();
+						}),
+					TAttribute<FText>(),
+							TAttribute<FText>(),
+							TAttribute<FSlateIcon>(),
+							true
+							);
+
+
+				Toolbar.AddToolBarButton(ToggleButtonAction, NAME_None, LOCTEXT("ToolBarButtonToggle", "Toggle"), FText::GetEmpty(), FSlateIcon(FAppStyle::GetAppStyleSetName(), "Icons.box-perspective"), EUserInterfaceActionType::ToggleButton);
+
+				Toolbar.AddComboButton(
+					FUIAction(),
+					FOnGetContent::CreateLambda(
+						[&]()
+						{
+							FMenuBuilder Menu(true, nullptr);
+							MakeMenuEntries(Menu);
+
+							return Menu.MakeWidget();
+						}),
+					TAttribute<FText>(),
+							TAttribute<FText>(),
+							TAttribute<FSlateIcon>(),
+							true
+							);
+
+
+			}
+			Toolbar.EndSection();
+		}
+
+		return Toolbar.MakeWidget();
+	}
+
     TSharedRef<SWidget> ConstructWidgetGallery()
     {
 
@@ -597,8 +683,7 @@ public:
             .HAlign(HAlign_Left)
             .VAlign(VAlign_Center)
             [
-                SNew(SCheckBox)
-                .Style(&FAppStyle::Get().GetWidgetStyle<FCheckBoxStyle>("SimplifiedCheckbox"))
+		SNew(SCheckBox)
                 .IsChecked_Lambda( [HBox] { return HBox->IsEnabled() ? ECheckBoxState::Checked : ECheckBoxState::Unchecked; } )
                 .OnCheckStateChanged_Lambda( [HBox] (ECheckBoxState NewState) { HBox->SetEnabled( NewState == ECheckBoxState::Checked); })
             ];
@@ -683,7 +768,7 @@ public:
                 .ButtonStyle( &FAppStyle::Get().GetWidgetStyle< FButtonStyle >( "PrimaryButton" ) )
                 .TextStyle( &FAppStyle::Get().GetWidgetStyle< FTextBlockStyle >("ButtonText"))
                 [
-                    LeftRightLabel(NAME_None, LOCTEXT("PrimaryButtonExampleLabel", "Primary Button"))
+                    LeftRightLabel(NAME_None, LOCTEXT("Label", "Label"))
                 ]
             ]
 
@@ -695,7 +780,7 @@ public:
                 SNew(SButton)
                 .ButtonStyle( &FAppStyle::Get().GetWidgetStyle< FButtonStyle >( "PrimaryButton" ) )
                 [
-                    LeftRightLabel("Icons.box-perspective", LOCTEXT("PrimaryButtonExampleLabel", "Primary Button"))
+                    LeftRightLabel("Icons.box-perspective", LOCTEXT("Label", "Label"))
                 ]
             ]
 
@@ -714,7 +799,7 @@ public:
         ];
 
         // SButton Rounded
-        NextSlot(WidgetGrid, LOCTEXT("SButtonLabelRounded", "Button"))
+        NextSlot(WidgetGrid, LOCTEXT("SButtonLabelRounded", "Default/Secondary Button"))
         [
 
             SNew(SHorizontalBox)
@@ -724,9 +809,9 @@ public:
             .VAlign(VAlign_Center)
             [
                 SNew(SButton)
-                .ButtonStyle( &FAppStyle::Get().GetWidgetStyle< FButtonStyle >( "Button" ) )
+                .ButtonStyle(&FAppStyle::Get().GetWidgetStyle<FButtonStyle>( "Button" ))
                 [
-                    LeftRightLabel(NAME_None, LOCTEXT("ButtonExampleLabel", "Normal Button"))
+                    LeftRightLabel(NAME_None, LOCTEXT("Label", "Label"), "SmallButtonText")
                 ]
             ]
 
@@ -737,7 +822,7 @@ public:
                 SNew(SButton)
                 .ButtonStyle( &FAppStyle::Get().GetWidgetStyle< FButtonStyle >( "Button" ) )
                 [
-                    LeftRightLabel("Icons.box-perspective", LOCTEXT("ButtonExampleLabel", "Normal Button"))
+                    LeftRightLabel("Icons.box-perspective", LOCTEXT("Label", "Label"), "SmallButtonText")
                 ]
             ]
 
@@ -758,92 +843,6 @@ public:
 
         // NoBorder Button Button 
         NextSlot(WidgetGrid, LOCTEXT("TextButton", "Simple Button"))
-        [
-
-            SNew(SHorizontalBox)
-
-            +SHorizontalBox::Slot()
-            .AutoWidth()
-            .Padding(8.f, 0.0f)
-            .VAlign(VAlign_Center)
-            [
-                SNew(SButton)
-                .ButtonStyle( &FAppStyle::Get().GetWidgetStyle< FButtonStyle >( "SimpleButton" ) )
-                [
-                    LeftRightLabel(NAME_None, LOCTEXT("TextButtonExampleLabel", "Simple Button"))
-                ]
-            ]
-
-            +SHorizontalBox::Slot()
-            .AutoWidth()
-            .Padding(8.f, 0.0f)
-            [
-                SNew(SButton)
-                .ButtonStyle( &FAppStyle::Get().GetWidgetStyle< FButtonStyle >( "SimpleButton" ) )
-                [
-                    LeftRightLabel("Icons.box-perspective", LOCTEXT("TextButtonExampleLabel", "Simple Button")) 
-                ]
-            ]
-
-            +SHorizontalBox::Slot()
-            .AutoWidth()
-            .Padding(8.f, 0.0f)
-            [
-                SNew(SButton)
-                .ButtonStyle( &FAppStyle::Get().GetWidgetStyle< FButtonStyle >( "SimpleButton" ) )
-                [
-                    LeftRightLabel("Icons.box-perspective")
-                ]
-            ]
-
-        ];
-
-        // SButton Dialog Primary Rounded
-        NextSlot(WidgetGrid, LOCTEXT("SButtonDialogPrimaryExampleLabelRounded", "Dialog Primary Button"))
-        [
-            SNew(SHorizontalBox)
-            +SHorizontalBox::Slot()
-            .AutoWidth()
-            .Padding(8.f, 0.0f)
-            .VAlign(VAlign_Center)
-            [
-
-                SNew(SButton)
-                .ButtonStyle( &FAppStyle::Get().GetWidgetStyle< FButtonStyle >( "PrimaryButton" ) )
-                .TextStyle( &FAppStyle::Get().GetWidgetStyle< FTextBlockStyle >("DialogButtonText"))
-                [
-                    LeftRightLabel(NAME_None, LOCTEXT("DialogPrimaryButtonExampleLabel", "Dialog Primary"), "DialogButtonText")
-                ]
-            ]
-
-            +SHorizontalBox::Slot()
-            .AutoWidth()
-            .Padding(8.f, 0.0f)
-            .VAlign(VAlign_Center)
-            [
-                SNew(SButton)
-                .ButtonStyle( &FAppStyle::Get().GetWidgetStyle< FButtonStyle >( "PrimaryButton" ) )
-                [
-                    LeftRightLabel("Icons.box-perspective", LOCTEXT("DialogPrimaryButtonExampleLabel", "Dialog Primary"), "DialogButtonText")
-                ]
-            ]
-
-            +SHorizontalBox::Slot()
-            .AutoWidth()
-            .Padding(8.f, 0.0f)
-            .VAlign(VAlign_Center)
-            [
-                SNew(SButton)
-                .ButtonStyle( &FAppStyle::Get().GetWidgetStyle< FButtonStyle >( "PrimaryButton" ) )
-                .VAlign(VAlign_Center)
-                [
-                    LeftRightLabel("Icons.box-perspective")
-                ]
-            ]
-        ];
-
-        // SButton Dialog 
-        NextSlot(WidgetGrid, LOCTEXT("SButtonLabelRounded", "Button"))
         [
 
             SNew(SHorizontalBox)
@@ -853,9 +852,9 @@ public:
             .VAlign(VAlign_Center)
             [
                 SNew(SButton)
-                .ButtonStyle( &FAppStyle::Get().GetWidgetStyle< FButtonStyle >( "Button" ) )
+                .ButtonStyle(&FAppStyle::Get().GetWidgetStyle<FButtonStyle>("SimpleButton"))
                 [
-                    LeftRightLabel(NAME_None, LOCTEXT("DialogButtonExampleLabel", "Dialog Button"), "DialogButtonText")
+                    LeftRightLabel(NAME_None, LOCTEXT("Label", "Label"), "SmallButtonText")
                 ]
             ]
 
@@ -864,9 +863,9 @@ public:
             .Padding(8.f, 0.0f)
             [
                 SNew(SButton)
-                .ButtonStyle( &FAppStyle::Get().GetWidgetStyle< FButtonStyle >( "Button" ) )
+                .ButtonStyle(&FAppStyle::Get().GetWidgetStyle<FButtonStyle>("SimpleButton"))
                 [
-                    LeftRightLabel("Icons.box-perspective", LOCTEXT("DialogButtonExampleLabel", "Dialog Button"), "DialogButtonText")
+                    LeftRightLabel("Icons.box-perspective", LOCTEXT("Label", "Label"), "SmallButtonText")
                 ]
             ]
 
@@ -875,8 +874,7 @@ public:
             .Padding(8.f, 0.0f)
             [
                 SNew(SButton)
-                .ButtonStyle( &FAppStyle::Get().GetWidgetStyle< FButtonStyle >( "Button" ) )
-                .VAlign(VAlign_Center)
+                .ButtonStyle(&FAppStyle::Get().GetWidgetStyle<FButtonStyle>("SimpleButton"))
                 [
                     LeftRightLabel("Icons.box-perspective")
                 ]
@@ -884,61 +882,6 @@ public:
 
         ];
 
-
-        // Toggle Button 
-        NextSlot(WidgetGrid, LOCTEXT("ToggleButtonAlt", "Toggle Button Alt"))
-        [
-
-            SNew(SHorizontalBox)
-
-            +SHorizontalBox::Slot()
-            .AutoWidth()
-            .Padding(8.f)
-            [
-                SNew(SCheckBox)
-                .Style( &FAppStyle::Get().GetWidgetStyle<FCheckBoxStyle>("ToggleButtonCheckBoxAlt"))
-                [
-                    LeftRightLabel("Icons.box-perspective")
-                ]
-            ]
-
-            +SHorizontalBox::Slot()
-            .AutoWidth()
-            .Padding(8.f)
-            [
-                SNew(SCheckBox)
-                .Style( &FAppStyle::Get().GetWidgetStyle<FCheckBoxStyle>("ToggleButtonCheckBoxAlt"))
-                .IsChecked(true)
-                [
-                    LeftRightLabel("Icons.pyramid")
-                ]
-            ]
-
-            +SHorizontalBox::Slot()
-            .AutoWidth()
-            .Padding(8.f)
-            [
-                SNew(SCheckBox)
-                .Style( &FAppStyle::Get().GetWidgetStyle<FCheckBoxStyle>("ToggleButtonCheckBoxAlt"))
-                [
-                    LeftRightLabel("Icons.cylinder")
-                ]
-            ]
-
-            +SHorizontalBox::Slot()
-            .AutoWidth()
-            .Padding(8.f)
-            [
-                SNew(SCheckBox)
-                .Style( &FAppStyle::Get().GetWidgetStyle<FCheckBoxStyle>("ToggleButtonCheckBoxAlt"))
-                .IsChecked(true)
-                [
-                    LeftRightLabel("Icons.sphere")
-                ]
-            ]
-
-        ];
-
         // Toggle Button with Words
         NextSlot(WidgetGrid, LOCTEXT("ToggleButton", "Toggle Button"))
         [
@@ -952,7 +895,7 @@ public:
                 .Style( &FAppStyle::Get().GetWidgetStyle<FCheckBoxStyle>("ToggleButtonCheckBox"))
                 .IsChecked(true)
                 [
-                    LeftRightLabel("Icons.box-perspective", LOCTEXT("Box", "BOX"))
+                    LeftRightLabel(NAME_None, LOCTEXT("Label", "Label"), "SmallButtonText")
                 ]
             ]
 
@@ -964,7 +907,7 @@ public:
                 .Style( &FAppStyle::Get().GetWidgetStyle<FCheckBoxStyle>("ToggleButtonCheckBox"))
                 .IsChecked(true)
                 [
-                    LeftRightLabel("Icons.pyramid", LOCTEXT("Pyramid", "PYRAMID"))
+                    LeftRightLabel("Icons.pyramid", LOCTEXT("Label", "Label"), "SmallButtonText")
                 ]
             ]
 
@@ -975,26 +918,12 @@ public:
                 SNew(SCheckBox)
                 .Style( &FAppStyle::Get().GetWidgetStyle<FCheckBoxStyle>("ToggleButtonCheckBox"))
                 [
-                    LeftRightLabel("Icons.cylinder", LOCTEXT("Cylinder", "CYLINDER"))
-                ]
-            ]
-
-
-
-            +SHorizontalBox::Slot()
-            .AutoWidth()
-            .Padding(8.f)
-            [
-                SNew(SCheckBox)
-                .Style( &FAppStyle::Get().GetWidgetStyle<FCheckBoxStyle>("ToggleButtonCheckBox"))
-                [
-                    LeftRightLabel("Icons.sphere", LOCTEXT("Sphere", "SPHERE"))
+                    LeftRightLabel("Icons.cylinder")
                 ]
             ]
-
         ];
 
-
+   
         // Segmented Control - Simple version lets you specify text and brush
         NextSlot(WidgetGrid, LOCTEXT("SegmentedControl", "SegmentedControl")) 
         [
@@ -1005,19 +934,19 @@ public:
 
         	+SSegmentedControl<int32>::Slot(0)
             .Icon(FAppStyle::Get().GetBrush("Icons.box-perspective"))
-            .Text(LOCTEXT("Box", "BOX"))
+            .Text(LOCTEXT("Box", "Box"))
 
             +SSegmentedControl<int32>::Slot(1)
             .Icon(FAppStyle::Get().GetBrush("Icons.cylinder"))
-            .Text(LOCTEXT("Cylinder", "CYLINDER"))
+            .Text(LOCTEXT("Cylinder", "Cylinder"))
 
             +SSegmentedControl<int32>::Slot(2)
             .Icon(FAppStyle::Get().GetBrush("Icons.pyramid"))
-            .Text(LOCTEXT("Pyramid", "PYRAMID"))
+            .Text(LOCTEXT("Pyramid", "Pyramid"))
 
         	+SSegmentedControl<int32>::Slot(3)
             .Icon(FAppStyle::Get().GetBrush("Icons.sphere"))
-            .Text(LOCTEXT("Sphere", "SPHERE"))
+            .Text(LOCTEXT("Sphere", "Sphere"))
 
         ];
    
@@ -1049,24 +978,6 @@ public:
 
         ];
 
-
-        // SCheckBox Simplified
-        NextSlot(WidgetGrid, LOCTEXT("SCheckBoxSimplifiedLabel", "Check Box Simple"))
-        [
-            SNew(SVerticalBox)
-
-            +SVerticalBox::Slot()
-            [
-                SNew(SHorizontalBox)
-                + SHorizontalBox::Slot().Padding(8.f) [ SNew(SCheckBox).Style( &FAppStyle::Get().GetWidgetStyle< FCheckBoxStyle > ("SimplifiedCheckbox")).IsChecked(ECheckBoxState::Unchecked) ]
-                + SHorizontalBox::Slot().Padding(8.f) [ SNew(SCheckBox).Style( &FAppStyle::Get().GetWidgetStyle< FCheckBoxStyle > ("SimplifiedCheckbox")).IsChecked(ECheckBoxState::Unchecked) ]
-                + SHorizontalBox::Slot().Padding(8.f) [ SNew(SCheckBox).Style( &FAppStyle::Get().GetWidgetStyle< FCheckBoxStyle > ("SimplifiedCheckbox")).IsChecked(ECheckBoxState::Checked) ]
-                + SHorizontalBox::Slot().Padding(8.f) [ SNew(SCheckBox).Style( &FAppStyle::Get().GetWidgetStyle< FCheckBoxStyle > ("SimplifiedCheckbox")).IsChecked(ECheckBoxState::Undetermined) ]
-            ]
-
-        ];
-
-
         // SCheckBox (as radio button)
         TSharedPtr<SHorizontalBox> RadioBox = SNew(SHorizontalBox);
         for (int i = 0; i < 5; i++)
@@ -1097,39 +1008,6 @@ public:
         ];
 
 
-        // SComboButton
-        /*
-        NextSlot(WidgetGrid, LOCTEXT("SComboButtonLabel", "SComboButton"))
-        .FillWidth(1.0)
-        [
-            SNew(SComboButton)
-            .ComboButtonStyle( &FAppStyle::Get().GetWidgetStyle< FComboButtonStyle >("SimpleComboButton"))
-            .ButtonContent()
-            [
-                SNew(STextBlock)
-                .Margin(FMargin(0.f, 0.f, 4.f, 0.f))
-                .Text(LOCTEXT("ComboButtonLabel", "Combo Button"))
-            ]
-            .MenuContent()
-            [
-                SNew(SVerticalBox)
-
-                + SVerticalBox::Slot().AutoHeight()[SNew(STextBlock).Text(LOCTEXT("ComboButtonItemLabel01", "Apple"))]
-                + SVerticalBox::Slot().AutoHeight().HAlign(HAlign_Fill)
-                [
-                    SNew(SButton)
-                    [
-                        SNew(STextBlock)
-                        .Text(LOCTEXT("ComboButtonItemLabel02", "Banana"))
-                    ]
-                ]
-
-                + SVerticalBox::Slot().AutoHeight()[SNew(STextBlock).Text(LOCTEXT("ComboButtonItemLabel03", "Carrot"))]
-                + SVerticalBox::Slot().AutoHeight()[SNew(STextBlock).Text(LOCTEXT("ComboButtonItemLabel06", "Diakon"))]
-
-            ]
-        ];*/
-
          // SComboBox Simple Icon Only version 
         NextSlot(WidgetGrid, LOCTEXT("SComboBoxIconLabel", "SimpleComboBox"))
         [
@@ -1169,7 +1047,7 @@ public:
                 }
             } )
             [
-                SAssignNew(ComboBoxTitleBlock, STextBlock).Text(LOCTEXT("ComboLabel", "Choose your Item"))  
+                SAssignNew(ComboBoxTitleBlock, STextBlock).Text(LOCTEXT("ComboLabel", "Label"))  
             ]
 
         ];
@@ -1389,7 +1267,6 @@ public:
         NextSlot(WidgetGrid, LOCTEXT("SNumericEntryBoxLabel", "SNumericEntryBox"))
         .FillWidth(1.0)
         [
-
             SNew(SNumericEntryBox<float>)
             .MinValue(-1000.0f)
             .MaxValue(1000.0f)
@@ -1435,12 +1312,6 @@ public:
         .FillWidth(1.0)
         [
             SNew(SEditableTextBox)
-            // .SetColorAndOpacity(FLinearColor::White)
-            // .ForegroundColor(FLinearColor::White)
-            // .BackgroundColor(FLinearColor::White)
-            // .SetTextBoxForegroundColor(FLinearColor::White)
-            // .SetTextBoxBackgroundColor(FLinearColor::White)
-            // .(FLinearColor::White)
             .HintText(LOCTEXT("SEditableTextBoxHint", "This is an editable text box"))
         ];
 
@@ -1453,6 +1324,7 @@ public:
             [
                 SNew(SMultiLineEditableText)
                 .HintText(LOCTEXT("SMultiLineEditableTextHint", "This is multi-line \n\t\t\t editable text"))
+		.OverflowPolicy(ETextOverflowPolicy::Ellipsis)
             ]
         ];
 
@@ -1464,6 +1336,7 @@ public:
             [
                 SNew(SMultiLineEditableTextBox)
                 .HintText(LOCTEXT("SMultiLineEditableTextBoxHint", "This is a multi-line editable text box"))
+		.OverflowPolicy(ETextOverflowPolicy::Ellipsis)
             ]    
         ];
 
@@ -1509,6 +1382,43 @@ public:
         ];
 
 
+		auto NewMenuDelegate = FNewMenuDelegate::CreateLambda(
+			[&](FMenuBuilder& MenuBuilder)
+			{
+				MenuBuilder.BeginSection(NAME_None, LOCTEXT("MenuHeader1", "Menu Header 1"));
+				MakeMenuEntries(MenuBuilder);
+				MenuBuilder.AddSubMenu(LOCTEXT("SubMenu1", "Sub Menu 1"), FText::GetEmpty(), FNewMenuDelegate::CreateRaw(this, &SStarshipGallery::MakeMenuEntries));
+				MenuBuilder.EndSection();
+
+				MenuBuilder.BeginSection(NAME_None, LOCTEXT("MenuHeader2", "Menu Header 2"));
+				MenuBuilder.AddMenuEntry(LOCTEXT("MenuEntry5", "Menu Entry 5"), FText::GetEmpty(), FSlateIcon(FAppStyle::GetAppStyleSetName(), "Icons.box-perspective"), FUIAction());
+				MenuBuilder.AddMenuEntry(LOCTEXT("MenuEntry6", "Menu Entry 6"), FText::GetEmpty(), FSlateIcon(FAppStyle::GetAppStyleSetName(), "Icons.box-perspective"), FUIAction());
+				MenuBuilder.AddMenuEntry(LOCTEXT("MenuEntry7", "Menu Entry 7"), FText::GetEmpty(), FSlateIcon(FAppStyle::GetAppStyleSetName(), "Icons.box-perspective"), FUIAction());
+				MenuBuilder.AddMenuEntry(LOCTEXT("MenuEntry8", "Menu Entry 8"), FText::GetEmpty(), FSlateIcon(FAppStyle::GetAppStyleSetName(), "Icons.box-perspective"), FUIAction());
+				MenuBuilder.EndSection();
+			});
+
+		FMenuBarBuilder MenuBar(nullptr);
+
+		MenuBar.AddPullDownMenu(LOCTEXT("Menu1", "Menu 1"), FText::GetEmpty(), NewMenuDelegate);
+		MenuBar.AddPullDownMenu(LOCTEXT("Menu2", "Menu 2"), FText::GetEmpty(), NewMenuDelegate);
+
+		NextSlot(WidgetGrid, LOCTEXT("MenuBar", "Menu Bar"))
+		[
+			MenuBar.MakeWidget()
+		];
+
+
+		NextSlot(WidgetGrid, LOCTEXT("Toolbar", "Toolbar"))
+		[
+			BuildToolbar(true)
+		];
+
+		NextSlot(WidgetGrid, LOCTEXT("Toolbar2", "Toolbar 2"))
+		[
+			BuildToolbar(false)
+		];
+
         return SNew(SBorder)
         .BorderImage( FAppStyle::Get().GetBrush("ToolPanel.GroupBorder") )
         [
