Added search filter for the list of icons in the Starship gallery.
Also fixed crash when selecting an option in the combo box.

diff --git a/Engine/Source/Runtime/AppFramework/Private/Framework/Testing/SStarshipGallery.cpp b/Engine/Source/Runtime/AppFramework/Private/Framework/Testing/SStarshipGallery.cpp
--- a/Engine/Source/Runtime/AppFramework/Private/Framework/Testing/SStarshipGallery.cpp
+++ b/Engine/Source/Runtime/AppFramework/Private/Framework/Testing/SStarshipGallery.cpp
@@ -361,10 +361,20 @@ public:
 
     };
 
+    TSharedPtr<SSearchBox> IconSearchBox;
+    bool IsIconVisible(const FString& IconPath) const
+    {
+        if (!IconSearchBox.IsValid() || IconSearchBox->GetText().IsEmpty())
+        {
+            return true;
+        }
+        return IconPath.Contains(IconSearchBox->GetText().ToString());
+    }
+
     TArray< TUniquePtr< FSlateBrush > > DynamicBrushes;
     TSharedRef<SWidget> ConstructIconsGallery()
     {
-        // auto GenerateColorButton = [] (FTex      t InTitle, FName StyleColorName, bool Inverted = false) -> TSharedRef<SWidget> 
+        // auto GenerateColorButton = [] (FText InTitle, FName StyleColorName, bool Inverted = false) -> TSharedRef<SWidget> 
 
         auto GenerateIconLibrary = [this] (FText InTitle, FString InPath) -> TSharedRef<SWidget>
         {
@@ -389,6 +399,7 @@ public:
                     SNew(SImage)
                     .Image(DynamicBrushes.Last().Get())
                     .ToolTipText( FText::FromString( IconPath ) )
+                     .Visibility_Lambda([this, IconPath]() { return IsIconVisible(IconPath) ? EVisibility::Visible : EVisibility::Collapsed; })
                 ];
             }
 
@@ -428,6 +439,7 @@ public:
                     SNew(SImage)
                     .Image(DynamicBrushes.Last().Get())
                     .ToolTipText( FText::FromString( IconPath ) )
+                    .Visibility_Lambda([this, IconPath]() { return IsIconVisible(IconPath) ? EVisibility::Visible : EVisibility::Collapsed; })
                 ];
             }
 
@@ -454,6 +466,15 @@ public:
             .Padding(48)
             [
                 SNew(SVerticalBox)
+
+                + SVerticalBox::Slot()
+                .AutoHeight()
+                .Padding(0.0f, 0.0f, 0.0f, 4.0f)
+                [
+                    SAssignNew(IconSearchBox, SSearchBox)
+                    .HintText(LOCTEXT("IconSearchHint", "Enter text to filter icons by path..."))
+                ]
+
                 +SVerticalBox::Slot().AutoHeight()[ GenerateIconLibrarySVG(NSLOCTEXT("StarshipGallery", "SlateCore", "Core"), "Content/Slate/Starship/Common")]
                 +SVerticalBox::Slot().AutoHeight()[ GenerateIconLibrarySVG(NSLOCTEXT("StarshipGallery", "Editor Common", "Editor"), "Content/Editor/Slate/Starship/Common")]
 
@@ -646,6 +667,8 @@ public:
 		return Toolbar.MakeWidget();
 	}
 
+    TSharedPtr<STextBlock> ComboBoxTitleBlock;
+
     TSharedRef<SWidget> ConstructWidgetGallery()
     {
 
@@ -1028,7 +1051,6 @@ public:
 
 
         // SComboBox
-        TSharedPtr<STextBlock> ComboBoxTitleBlock;
         NextSlot(WidgetGrid, LOCTEXT("SComboBoxLabel", "ComboBox"))
         .AutoWidth()
         // .FillWidth(1.0)
@@ -1039,9 +1061,9 @@ public:
             { 
                 return SNew(STextBlock).Text( FText::FromString(*Item));
             } )
-            .OnSelectionChanged_Lambda([ComboBoxTitleBlock] (TSharedPtr<FString> InSelection, ESelectInfo::Type InSelectInfo) 
+            .OnSelectionChanged_Lambda([this] (TSharedPtr<FString> InSelection, ESelectInfo::Type InSelectInfo) 
             {
-                if (InSelection.IsValid())
+                if (InSelection.IsValid() && ComboBoxTitleBlock.IsValid())
                 {
                     ComboBoxTitleBlock->SetText( FText::FromString(*InSelection));
                 }
