Add Noto Sans Math as an editor-only font.

diff --git a/Engine/Source/Runtime/AppFramework/Private/Framework/Testing/SStarshipGallery.cpp b/Engine/Source/Runtime/AppFramework/Private/Framework/Testing/SStarshipGallery.cpp
--- a/Engine/Source/Runtime/AppFramework/Private/Framework/Testing/SStarshipGallery.cpp
+++ b/Engine/Source/Runtime/AppFramework/Private/Framework/Testing/SStarshipGallery.cpp
@@ -572,6 +572,31 @@ public:
                 ]
             ]
 
+			+ SHorizontalBox::Slot()
+			.AutoWidth()
+			.Padding(64.f, 0.0f)
+			[
+				SNew(SVerticalBox)
+
+				+ SVerticalBox::Slot().Padding(0.f, 64.f, 0.f, 4.f).AutoHeight()[SNew(STextBlock).ColorAndOpacity(LabelColor).Text(NSLOCTEXT("StarshipGallery", "SubMathTextDesc", "SUB-TYPEFACES - MATH"))]
+				+ SVerticalBox::Slot().Padding(0.f, 4.f).AutoHeight()
+				[
+					SNew(STextBlock)
+					.Text(NSLOCTEXT("StarshipGallery", "MathematicalTextSpecimen", "Mathematical Alphanumeric Symbols, U+1D400 - U+1D7FF\n\U0001D400\U0001D401\U0001D402 \U0001D434\U0001D435\U0001D436 \U0001D4D0\U0001D4D1\U0001D4D2"))
+				]
+
+				+ SVerticalBox::Slot().Padding(0.f, 64.f, 0.f, 4.f).AutoHeight()[SNew(STextBlock).ColorAndOpacity(LabelColor).Text(NSLOCTEXT("StarshipGallery", "SubEmojiTextDesc", "SUB-TYPEFACES - EMOJI"))]
+				+ SVerticalBox::Slot().Padding(0.f, 4.f).AutoHeight()
+				[
+					SNew(STextBlock)
+					.Text(NSLOCTEXT("StarshipGallery", "EmoticonsTextSpecimen", "Emoticons, U+1F600 - U+1F64F\n\U0001F60E\U0001F643\U0001F648"))
+				]
+				+ SVerticalBox::Slot().Padding(0.f, 4.f).AutoHeight()
+				[
+					SNew(STextBlock)
+					.Text(NSLOCTEXT("StarshipGallery", "SupplementalSymbolsPictographsTextSpecimen", "Supplemental Symbols and Pictographs, U+1F900 - U+1F9FF\n\U0001F914\U0001F916\U0001F9F2"))
+				]
+			]
         ];
     }
 
diff --git a/Engine/Source/Runtime/SlateCore/Private/Fonts/LegacySlateFontInfoCache.cpp b/Engine/Source/Runtime/SlateCore/Private/Fonts/LegacySlateFontInfoCache.cpp
--- a/Engine/Source/Runtime/SlateCore/Private/Fonts/LegacySlateFontInfoCache.cpp
+++ b/Engine/Source/Runtime/SlateCore/Private/Fonts/LegacySlateFontInfoCache.cpp
@@ -197,6 +197,14 @@ TSharedRef<const FCompositeFont> FLegacySlateFontInfoCache::GetDefaultFont()
 			APPEND_EDITOR_FONT(SubFont.Typeface, "VeryLight", "NanumGothic.ttf", EFontHinting::Default);
 		}
 
+		// Math (editor-only)
+		if (GIsEditor)
+		{
+			FCompositeSubFont& SubFont = MutableDefaultFont->SubTypefaces.AddDefaulted_GetRef();
+			APPEND_RANGE(SubFont, MathematicalAlphanumericSymbols);
+			APPEND_EDITOR_FONT(SubFont.Typeface, "Regular", "NotoSansMath-Regular.ttf", EFontHinting::Default);
+		}
+
 		// Emoji (editor-only)
 		if (GIsEditor)
 		{
