Remove c-style cast that is doing a reinterpret_cast when the type is not included by the callee.

diff --git a/Engine/Source/Editor/Kismet/Private/Debugging/SKismetDebuggingView.cpp b/Engine/Source/Editor/Kismet/Private/Debugging/SKismetDebuggingView.cpp
--- a/Engine/Source/Editor/Kismet/Private/Debugging/SKismetDebuggingView.cpp
+++ b/Engine/Source/Editor/Kismet/Private/Debugging/SKismetDebuggingView.cpp
@@ -65,7 +65,7 @@ TSharedRef<SHorizontalBox> SKismetDebuggingView::GetDebugLineTypeToggle(FDebugLi
 		.AutoWidth()
 		[
 			SNew(SCheckBox)
-				.IsChecked(true)
+				.IsChecked(ECheckBoxState::Checked)
 				.OnCheckStateChanged_Static(&FDebugLineItem::OnDebugLineTypeActiveChanged, Type)
 		]
 		+ SHorizontalBox::Slot()
diff --git a/Engine/Source/Runtime/AppFramework/Private/Framework/Testing/SStarshipGallery.cpp b/Engine/Source/Runtime/AppFramework/Private/Framework/Testing/SStarshipGallery.cpp
--- a/Engine/Source/Runtime/AppFramework/Private/Framework/Testing/SStarshipGallery.cpp
+++ b/Engine/Source/Runtime/AppFramework/Private/Framework/Testing/SStarshipGallery.cpp
@@ -941,7 +941,7 @@ public:
             [
                 SNew(SCheckBox)
                 .Style( &FAppStyle::Get().GetWidgetStyle<FCheckBoxStyle>("ToggleButtonCheckBox"))
-                .IsChecked(true)
+                .IsChecked(ECheckBoxState::Checked)
                 [
                     LeftRightLabel(NAME_None, LOCTEXT("Label", "Label"), "SmallButtonText")
                 ]
@@ -953,7 +953,7 @@ public:
             [
                 SNew(SCheckBox)
                 .Style( &FAppStyle::Get().GetWidgetStyle<FCheckBoxStyle>("ToggleButtonCheckBox"))
-                .IsChecked(true)
+                .IsChecked(ECheckBoxState::Checked)
                 [
                     LeftRightLabel("Icons.pyramid", LOCTEXT("Label", "Label"), "SmallButtonText")
                 ]
diff --git a/Engine/Source/Runtime/Core/Public/Misc/Attribute.h b/Engine/Source/Runtime/Core/Public/Misc/Attribute.h
--- a/Engine/Source/Runtime/Core/Public/Misc/Attribute.h
+++ b/Engine/Source/Runtime/Core/Public/Misc/Attribute.h
@@ -42,7 +42,7 @@ public:
 	 */
 	template< typename OtherType >
 	TAttribute( const OtherType& InInitialValue )
-		: Value( (ObjectType)InInitialValue )
+		: Value( static_cast<ObjectType>(InInitialValue) )
 		, bIsSet(true)
 		, Getter()
 	{
diff --git a/Engine/Source/Runtime/Slate/Private/Widgets/Views/SHeaderRow.cpp b/Engine/Source/Runtime/Slate/Private/Widgets/Views/SHeaderRow.cpp
--- a/Engine/Source/Runtime/Slate/Private/Widgets/Views/SHeaderRow.cpp
+++ b/Engine/Source/Runtime/Slate/Private/Widgets/Views/SHeaderRow.cpp
@@ -814,7 +814,7 @@ void SHeaderRow::RegenerateWidgets()
 					const float GripSize = SplitterHandleSize > 0.0f ? SplitterHandleSize + 4.0f : 5.0f;
 					TSharedRef<SBorder> SizingGrip = SNew(SBorder)
 						.Padding(0.0f)
-						.BorderImage( &Style->ColumnSplitterStyle.HandleNormalBrush )
+						.BorderImage( nullptr )
 						.Cursor(EMouseCursor::ResizeLeftRight)
 						.Content()
 						[
