Deprecate arrays from SColorPicker. Memstomp was fixed with 23672712. This CL removed the posibility of creating a new memstomp

diff --git a/Engine/Plugins/Enterprise/VariantManager/Source/VariantManager/Private/DisplayNodes/VariantManagerColorPropertyNode.cpp b/Engine/Plugins/Enterprise/VariantManager/Source/VariantManager/Private/DisplayNodes/VariantManagerColorPropertyNode.cpp
--- a/Engine/Plugins/Enterprise/VariantManager/Source/VariantManager/Private/DisplayNodes/VariantManagerColorPropertyNode.cpp
+++ b/Engine/Plugins/Enterprise/VariantManager/Source/VariantManager/Private/DisplayNodes/VariantManagerColorPropertyNode.cpp
@@ -107,7 +107,7 @@ FReply FVariantManagerColorPropertyNode::OnClickColorBlock(const FGeometry& MyGe
 		PickerArgs.OnColorCommitted = FOnLinearColorValueChanged::CreateSP(this, &FVariantManagerColorPropertyNode::OnSetColorFromColorPicker);
 	}
 
-	FVariantColorNodeImpl::GetCommonColorFromPropertyValues(PickerArgs.InitialColorOverride, PropertyValues);
+	FVariantColorNodeImpl::GetCommonColorFromPropertyValues(PickerArgs.InitialColor, PropertyValues);
 
 	OpenColorPicker(PickerArgs);
 
diff --git a/Engine/Plugins/Experimental/Mutable/Source/CustomizableObjectEditor/Private/MuCOE/CustomizableObjectCustomSettings.cpp b/Engine/Plugins/Experimental/Mutable/Source/CustomizableObjectEditor/Private/MuCOE/CustomizableObjectCustomSettings.cpp
--- a/Engine/Plugins/Experimental/Mutable/Source/CustomizableObjectEditor/Private/MuCOE/CustomizableObjectCustomSettings.cpp
+++ b/Engine/Plugins/Experimental/Mutable/Source/CustomizableObjectEditor/Private/MuCOE/CustomizableObjectCustomSettings.cpp
@@ -725,7 +725,7 @@ FReply SCustomizableObjectCustomSettings::OnLightColorBlockMouseButtonDown(const
 	args.bIsModal = true;
 	args.bUseAlpha = false;
 	args.bOnlyRefreshOnMouseUp = false;
-	args.InitialColorOverride = GetLightColorValue();
+	args.InitialColor = GetLightColorValue();
 	args.OnColorCommitted = FOnLinearColorValueChanged::CreateSP(this, &SCustomizableObjectCustomSettings::OnSetLightColorFromColorPicker);
 
 	OpenColorPicker(args);
diff --git a/Engine/Plugins/Experimental/Mutable/Source/CustomizableObjectEditor/Private/MuCOE/SCustomizableInstanceProperties.cpp b/Engine/Plugins/Experimental/Mutable/Source/CustomizableObjectEditor/Private/MuCOE/SCustomizableInstanceProperties.cpp
--- a/Engine/Plugins/Experimental/Mutable/Source/CustomizableObjectEditor/Private/MuCOE/SCustomizableInstanceProperties.cpp
+++ b/Engine/Plugins/Experimental/Mutable/Source/CustomizableObjectEditor/Private/MuCOE/SCustomizableInstanceProperties.cpp
@@ -1697,7 +1697,7 @@ FReply SCustomizableInstanceProperties::OnColorBlockMouseButtonDown(const FGeome
 	args.bIsModal = true;
 	args.bUseAlpha = false;
 	args.bOnlyRefreshOnMouseUp = false;
-	args.InitialColorOverride = col;
+	args.InitialColor = col;
 	args.OnColorCommitted = FOnLinearColorValueChanged::CreateSP(this, &SCustomizableInstanceProperties::OnSetColorFromColorPicker, ParamName);
 	OpenColorPicker(args);
 
diff --git a/Engine/Plugins/Experimental/Mutable/Source/CustomizableObjectEditor/Private/MuCOE/SMutableParametersWidget.cpp b/Engine/Plugins/Experimental/Mutable/Source/CustomizableObjectEditor/Private/MuCOE/SMutableParametersWidget.cpp
--- a/Engine/Plugins/Experimental/Mutable/Source/CustomizableObjectEditor/Private/MuCOE/SMutableParametersWidget.cpp
+++ b/Engine/Plugins/Experimental/Mutable/Source/CustomizableObjectEditor/Private/MuCOE/SMutableParametersWidget.cpp
@@ -437,7 +437,7 @@ FReply SMutableParametersWidget::OnColorBlockMouseButtonDown(const FGeometry& My
 	args.bIsModal = true;
 	args.bUseAlpha = false;
 	args.bOnlyRefreshOnMouseUp = false;
-	args.InitialColorOverride = col;
+	args.InitialColor = col;
 	args.OnColorCommitted = FOnLinearColorValueChanged::CreateSP(this, &SMutableParametersWidget::OnSetColorFromColorPicker, ParamIndex);
 	OpenColorPicker(args);
 
diff --git a/Engine/Plugins/Experimental/Mutable/Source/CustomizableObjectPopulationEditor/Private/MuCOPE/CustomizableObjectPopulationClassDetails.cpp b/Engine/Plugins/Experimental/Mutable/Source/CustomizableObjectPopulationEditor/Private/MuCOPE/CustomizableObjectPopulationClassDetails.cpp
--- a/Engine/Plugins/Experimental/Mutable/Source/CustomizableObjectPopulationEditor/Private/MuCOPE/CustomizableObjectPopulationClassDetails.cpp
+++ b/Engine/Plugins/Experimental/Mutable/Source/CustomizableObjectPopulationEditor/Private/MuCOPE/CustomizableObjectPopulationClassDetails.cpp
@@ -1499,7 +1499,7 @@ FReply SPopulationClassConstraint::OnColorPreviewClicked(const FGeometry& MyGeom
 	args.bIsModal = true;
 	args.bUseAlpha = true;
 	args.bOnlyRefreshOnMouseUp = false;
-	args.InitialColorOverride = PopulationClass->Characteristics[CharactericticIndex].Constraints[ConstraintIndex].DiscreteColor;
+	args.InitialColor = PopulationClass->Characteristics[CharactericticIndex].Constraints[ConstraintIndex].DiscreteColor;
 	args.OnColorCommitted = FOnLinearColorValueChanged::CreateSP(this, &SPopulationClassConstraint::OnSetConstantColorFromColorPicker);
 
 	OpenColorPicker(args);
diff --git a/Engine/Plugins/Experimental/PCG/Source/PCGEditor/Private/PCGEditorGraphNodeBase.cpp b/Engine/Plugins/Experimental/PCG/Source/PCGEditor/Private/PCGEditorGraphNodeBase.cpp
--- a/Engine/Plugins/Experimental/PCG/Source/PCGEditor/Private/PCGEditorGraphNodeBase.cpp
+++ b/Engine/Plugins/Experimental/PCG/Source/PCGEditor/Private/PCGEditorGraphNodeBase.cpp
@@ -262,7 +262,7 @@ void UPCGEditorGraphNodeBase::OnPickColor()
 	FColorPickerArgs PickerArgs;
 	PickerArgs.bIsModal = true;
 	PickerArgs.bUseAlpha = false;
-	PickerArgs.InitialColorOverride = GetNodeTitleColor();
+	PickerArgs.InitialColor = GetNodeTitleColor();
 	PickerArgs.OnColorCommitted = FOnLinearColorValueChanged::CreateUObject(this, &UPCGEditorGraphNodeBase::OnColorPicked);
 
 	OpenColorPicker(PickerArgs);
diff --git a/Engine/Plugins/FX/Niagara/Source/NiagaraEditor/Private/TypeEditorUtilities/NiagaraColorTypeEditorUtilities.cpp b/Engine/Plugins/FX/Niagara/Source/NiagaraEditor/Private/TypeEditorUtilities/NiagaraColorTypeEditorUtilities.cpp
--- a/Engine/Plugins/FX/Niagara/Source/NiagaraEditor/Private/TypeEditorUtilities/NiagaraColorTypeEditorUtilities.cpp
+++ b/Engine/Plugins/FX/Niagara/Source/NiagaraEditor/Private/TypeEditorUtilities/NiagaraColorTypeEditorUtilities.cpp
@@ -147,7 +147,7 @@ private:
 			PickerArgs.OnInteractivePickBegin = FSimpleDelegate::CreateSP(this, &SNiagaraColorParameterEditor::InteractivePickBegin);
 			PickerArgs.OnInteractivePickEnd = FSimpleDelegate::CreateSP(this, &SNiagaraColorParameterEditor::InteractivePickEnd);
 			PickerArgs.OnColorPickerWindowClosed = FOnWindowClosed::CreateSP(this, &SNiagaraColorParameterEditor::ColorPickerClosed);
-			PickerArgs.InitialColorOverride = ColorValue;
+			PickerArgs.InitialColor = ColorValue;
 			PickerArgs.ParentWidget = ColorBlock;
 		}
 
diff --git a/Engine/Plugins/Runtime/Metasound/Source/MetasoundEditor/Private/SMetasoundPinValueInspector.cpp b/Engine/Plugins/Runtime/Metasound/Source/MetasoundEditor/Private/SMetasoundPinValueInspector.cpp
--- a/Engine/Plugins/Runtime/Metasound/Source/MetasoundEditor/Private/SMetasoundPinValueInspector.cpp
+++ b/Engine/Plugins/Runtime/Metasound/Source/MetasoundEditor/Private/SMetasoundPinValueInspector.cpp
@@ -426,7 +426,7 @@ namespace Metasound
 			}
 
 			FColorPickerArgs PickerArgs = InitPickerArgs();
-			PickerArgs.InitialColorOverride = InInitColor;
+			PickerArgs.InitialColor = InInitColor;
 			PickerArgs.OnColorCommitted = FOnLinearColorValueChanged::CreateSP(this, &FMetasoundNumericDebugLineItem::OnColorCommitted, InIndex);
 
 			OpenColorPicker(PickerArgs);
diff --git a/Engine/Plugins/Runtime/ZoneGraph/Source/ZoneGraphEditor/Private/ZoneGraphTagInfoDetails.cpp b/Engine/Plugins/Runtime/ZoneGraph/Source/ZoneGraphEditor/Private/ZoneGraphTagInfoDetails.cpp
--- a/Engine/Plugins/Runtime/ZoneGraph/Source/ZoneGraphEditor/Private/ZoneGraphTagInfoDetails.cpp
+++ b/Engine/Plugins/Runtime/ZoneGraph/Source/ZoneGraphEditor/Private/ZoneGraphTagInfoDetails.cpp
@@ -95,7 +95,7 @@ FReply FZoneGraphTagInfoDetails::OnColorPressed(const FGeometry& MyGeometry, con
 	PickerArgs.DisplayGamma = TAttribute<float>::Create(TAttribute<float>::FGetter::CreateUObject(GEngine, &UEngine::GetDisplayGamma));
 	PickerArgs.OnColorCommitted = FOnLinearColorValueChanged::CreateSP(this, &FZoneGraphTagInfoDetails::SetColor);
 	PickerArgs.OnColorPickerCancelled = FOnColorPickerCancelled::CreateSP(this, &FZoneGraphTagInfoDetails::OnColorPickerCancelled);
-	PickerArgs.InitialColorOverride = Color.Get(FColor(128, 128, 128));
+	PickerArgs.InitialColor = Color.Get(FColor(128, 128, 128));
 	OpenColorPicker(PickerArgs);
 
 	return FReply::Handled();
diff --git a/Engine/Plugins/VirtualProduction/Takes/Source/TakeRecorder/Private/Widgets/STakeRecorderSources.cpp b/Engine/Plugins/VirtualProduction/Takes/Source/TakeRecorder/Private/Widgets/STakeRecorderSources.cpp
--- a/Engine/Plugins/VirtualProduction/Takes/Source/TakeRecorder/Private/Widgets/STakeRecorderSources.cpp
+++ b/Engine/Plugins/VirtualProduction/Takes/Source/TakeRecorder/Private/Widgets/STakeRecorderSources.cpp
@@ -307,7 +307,7 @@ private:
 		FColorPickerArgs PickerArgs;
 		PickerArgs.bUseAlpha = false;
 		PickerArgs.DisplayGamma = TAttribute<float>::Create(TAttribute<float>::FGetter::CreateUObject(GEngine, &UEngine::GetDisplayGamma));
-		PickerArgs.InitialColorOverride = GetSourceTintColor();
+		PickerArgs.InitialColor = GetSourceTintColor();
 		PickerArgs.OnColorCommitted = FOnLinearColorValueChanged::CreateSP(this, &FTakeRecorderSourceTreeItem::OnColorPickerPicked);
 
 		OpenColorPicker(PickerArgs);
diff --git a/Engine/Source/Developer/TraceInsights/Private/Insights/MemoryProfiler/Widgets/SMemTagTreeView.cpp b/Engine/Source/Developer/TraceInsights/Private/Insights/MemoryProfiler/Widgets/SMemTagTreeView.cpp
--- a/Engine/Source/Developer/TraceInsights/Private/Insights/MemoryProfiler/Widgets/SMemTagTreeView.cpp
+++ b/Engine/Source/Developer/TraceInsights/Private/Insights/MemoryProfiler/Widgets/SMemTagTreeView.cpp
@@ -2428,7 +2428,7 @@ void SMemTagTreeView::EditColorForSelectedMemTags()
 			//PickerArgs.OnInteractivePickBegin = FSimpleDelegate::CreateSP(this, &SMemTagTreeView::InteractivePickBegin);
 			//PickerArgs.OnInteractivePickEnd = FSimpleDelegate::CreateSP(this, &SMemTagTreeView::InteractivePickEnd);
 			//PickerArgs.OnColorPickerWindowClosed = FOnWindowClosed::CreateSP(this, &SMemTagTreeView::ColorPickerClosed);
-			PickerArgs.InitialColorOverride = EditableColorValue;
+			PickerArgs.InitialColor = EditableColorValue;
 			PickerArgs.ParentWidget = SharedThis(this);
 		}
 
diff --git a/Engine/Source/Editor/Cascade/Private/Cascade.cpp b/Engine/Source/Editor/Cascade/Private/Cascade.cpp
--- a/Engine/Source/Editor/Cascade/Private/Cascade.cpp
+++ b/Engine/Source/Editor/Cascade/Private/Cascade.cpp
@@ -4079,6 +4079,7 @@ void FCascade::OnBackgroundColor()
 						PinnedParticleSystem->BackgroundColor = NewValue.ToFColorSRGB();
 					}
 				}));
+			PickerArgs.bClampValue = true;
 			PickerArgs.ParentWidget = PreviewViewport;
 			PickerArgs.DisplayGamma = TAttribute<float>::Create(TAttribute<float>::FGetter::CreateUObject(GEngine, &UEngine::GetDisplayGamma));
 
diff --git a/Engine/Source/Editor/Cascade/Private/CascadeEmitterCanvasClient.cpp b/Engine/Source/Editor/Cascade/Private/CascadeEmitterCanvasClient.cpp
--- a/Engine/Source/Editor/Cascade/Private/CascadeEmitterCanvasClient.cpp
+++ b/Engine/Source/Editor/Cascade/Private/CascadeEmitterCanvasClient.cpp
@@ -402,6 +402,7 @@ bool FCascadeEmitterCanvasClient::InputKey(const FInputKeyEventArgs& EventArgs)
 
 										}
 									}));
+								PickerArgs.bClampValue = true;
 								PickerArgs.DisplayGamma = TAttribute<float>::Create( TAttribute<float>::FGetter::CreateUObject(GEngine, &UEngine::GetDisplayGamma) );
 
 								OpenColorPicker(PickerArgs);
diff --git a/Engine/Source/Editor/ContentBrowser/Private/CollectionContextMenu.cpp b/Engine/Source/Editor/ContentBrowser/Private/CollectionContextMenu.cpp
--- a/Engine/Source/Editor/ContentBrowser/Private/CollectionContextMenu.cpp
+++ b/Engine/Source/Editor/ContentBrowser/Private/CollectionContextMenu.cpp
@@ -612,7 +612,7 @@ void FCollectionContextMenu::ExecutePickColor()
 	FColorPickerArgs PickerArgs;
 	PickerArgs.bIsModal = true; // TODO: Allow live color updates via a proxy?
 	PickerArgs.ParentWidget = CollectionViewPtr;
-	PickerArgs.InitialColorOverride = InitialColor;
+	PickerArgs.InitialColor = InitialColor;
 	PickerArgs.OnColorCommitted.BindSP(this, &FCollectionContextMenu::OnColorCommitted);
 
 	OpenColorPicker(PickerArgs);
diff --git a/Engine/Source/Editor/CurveEditor/Private/CurveEditor.cpp b/Engine/Source/Editor/CurveEditor/Private/CurveEditor.cpp
--- a/Engine/Source/Editor/CurveEditor/Private/CurveEditor.cpp
+++ b/Engine/Source/Editor/CurveEditor/Private/CurveEditor.cpp
@@ -1914,7 +1914,7 @@ void FCurveEditor::SetCurveColorsForSelected()
 
 	FColorPickerArgs PickerArgs;
 	PickerArgs.bUseAlpha = false;
-	PickerArgs.InitialColorOverride = FindCurve(*CurveModelIDs.CreateIterator())->GetColor();
+	PickerArgs.InitialColor = FindCurve(*CurveModelIDs.CreateIterator())->GetColor();
 	PickerArgs.OnColorCommitted.BindLambda([WeakSelf, CurveModelIDs](FLinearColor NewColor)
 	{
 		if (TSharedPtr<FCurveEditor> Self = WeakSelf.Pin())
diff --git a/Engine/Source/Editor/DetailCustomizations/Private/ColorStructCustomization.cpp b/Engine/Source/Editor/DetailCustomizations/Private/ColorStructCustomization.cpp
--- a/Engine/Source/Editor/DetailCustomizations/Private/ColorStructCustomization.cpp
+++ b/Engine/Source/Editor/DetailCustomizations/Private/ColorStructCustomization.cpp
@@ -247,7 +247,7 @@ void FColorStructCustomization::CreateColorPicker(bool bUseAlpha)
 		PickerArgs.OnColorPickerWindowClosed = FOnWindowClosed::CreateSP(this, &FColorStructCustomization::OnColorPickerWindowClosed);
 		PickerArgs.OnInteractivePickBegin = FSimpleDelegate::CreateSP(this, &FColorStructCustomization::OnColorPickerInteractiveBegin);
 		PickerArgs.OnInteractivePickEnd = FSimpleDelegate::CreateSP(this, &FColorStructCustomization::OnColorPickerInteractiveEnd);
-		PickerArgs.InitialColorOverride = InitialColor;
+		PickerArgs.InitialColor = InitialColor;
 		PickerArgs.ParentWidget = ColorPickerParentWidget;
 		PickerArgs.OptionalOwningDetailsView = ColorPickerParentWidget;
 		FWidgetPath ParentWidgetPath;
diff --git a/Engine/Source/Editor/DetailCustomizations/Private/CustomPrimitiveDataCustomization.cpp b/Engine/Source/Editor/DetailCustomizations/Private/CustomPrimitiveDataCustomization.cpp
--- a/Engine/Source/Editor/DetailCustomizations/Private/CustomPrimitiveDataCustomization.cpp
+++ b/Engine/Source/Editor/DetailCustomizations/Private/CustomPrimitiveDataCustomization.cpp
@@ -897,7 +897,7 @@ FReply FCustomPrimitiveDataCustomization::OnMouseButtonDownColorBlock(const FGeo
 
 	FColorPickerArgs PickerArgs;
 	PickerArgs.bUseAlpha = true;
-	PickerArgs.InitialColorOverride = GetVectorColor(PrimIdx);
+	PickerArgs.InitialColor = GetVectorColor(PrimIdx);
 	PickerArgs.ParentWidget = ColorBlocks[PrimIdx];
 	PickerArgs.DisplayGamma = TAttribute<float>::Create(TAttribute<float>::FGetter::CreateUObject(GEngine, &UEngine::GetDisplayGamma));
 	PickerArgs.OnColorCommitted = FOnLinearColorValueChanged::CreateSP(this, &FCustomPrimitiveDataCustomization::SetVectorColor, PrimIdx);
diff --git a/Engine/Source/Editor/DetailCustomizations/Private/RawDistributionVectorStructCustomization.cpp b/Engine/Source/Editor/DetailCustomizations/Private/RawDistributionVectorStructCustomization.cpp
--- a/Engine/Source/Editor/DetailCustomizations/Private/RawDistributionVectorStructCustomization.cpp
+++ b/Engine/Source/Editor/DetailCustomizations/Private/RawDistributionVectorStructCustomization.cpp
@@ -488,7 +488,7 @@ void FReplaceVectorWithLinearColorBuilder::CreateColorPicker(const TSharedPtr<IP
 	PickerArgs.OnColorPickerCancelled = FOnColorPickerCancelled::CreateSP(this, &FReplaceVectorWithLinearColorBuilder::OnColorPickerCancelled, StructHandle);
 	PickerArgs.OnInteractivePickBegin = FSimpleDelegate::CreateSP(this, &FReplaceVectorWithLinearColorBuilder::OnColorPickerInteractiveBegin);
 	PickerArgs.OnInteractivePickEnd = FSimpleDelegate::CreateSP(this, &FReplaceVectorWithLinearColorBuilder::OnColorPickerInteractiveEnd);
-	PickerArgs.InitialColorOverride = OldColorValue;
+	PickerArgs.InitialColor = OldColorValue;
 	PickerArgs.ParentWidget = ColorPickerParentWidget;
 
 	OpenColorPicker(PickerArgs);
diff --git a/Engine/Source/Editor/DistCurveEditor/Private/SDistributionCurveEditor.cpp b/Engine/Source/Editor/DistCurveEditor/Private/SDistributionCurveEditor.cpp
--- a/Engine/Source/Editor/DistCurveEditor/Private/SDistributionCurveEditor.cpp
+++ b/Engine/Source/Editor/DistCurveEditor/Private/SDistributionCurveEditor.cpp
@@ -653,6 +653,7 @@ void SDistributionCurveEditor::OnSetColor()
 	FLinearColor TempColor = InputColor;
 	FColorPickerArgs PickerArgs = FColorPickerArgs(InputColor, FOnLinearColorValueChanged::CreateLambda([&TempColor](FLinearColor NewValue){ TempColor = NewValue.ToFColorSRGB(); }));
 	PickerArgs.bIsModal = true;
+	PickerArgs.bClampValue = true;
 	PickerArgs.DisplayGamma = TAttribute<float>::Create( TAttribute<float>::FGetter::CreateUObject(GEngine, &UEngine::GetDisplayGamma) );
 
 	if (OpenColorPicker(PickerArgs))
diff --git a/Engine/Source/Editor/FontEditor/Private/FontEditor.cpp b/Engine/Source/Editor/FontEditor/Private/FontEditor.cpp
--- a/Engine/Source/Editor/FontEditor/Private/FontEditor.cpp
+++ b/Engine/Source/Editor/FontEditor/Private/FontEditor.cpp
@@ -928,8 +928,9 @@ void FFontEditor::OnBackgroundColor()
 	PickerArgs.bIsModal = true;
 	PickerArgs.ParentWidget = FontPreview;
 	PickerArgs.bUseAlpha = true;
+	PickerArgs.bClampValue = true;
 	PickerArgs.DisplayGamma = TAttribute<float>::Create( TAttribute<float>::FGetter::CreateUObject(GEngine, &UEngine::GetDisplayGamma) );
-	PickerArgs.InitialColorOverride = FontPreviewWidget->GetPreviewBackgroundColor();
+	PickerArgs.InitialColor = FontPreviewWidget->GetPreviewBackgroundColor();
 	PickerArgs.OnColorCommitted = FOnLinearColorValueChanged::CreateLambda([WeakFontPreviewWidget](FLinearColor NewValue)
 		{
 			if (TSharedPtr<SFontEditorViewport> PinnedWidget = WeakFontPreviewWidget.Pin())
@@ -956,7 +957,7 @@ void FFontEditor::OnForegroundColor()
 	PickerArgs.ParentWidget = FontPreview;
 	PickerArgs.bUseAlpha = true;
 	PickerArgs.DisplayGamma = TAttribute<float>::Create( TAttribute<float>::FGetter::CreateUObject(GEngine, &UEngine::GetDisplayGamma) );
-	PickerArgs.InitialColorOverride = FontPreviewWidget->GetPreviewForegroundColor();
+	PickerArgs.InitialColor = FontPreviewWidget->GetPreviewForegroundColor();
 	PickerArgs.OnColorCommitted = FOnLinearColorValueChanged::CreateLambda([WeakFontPreviewWidget](FLinearColor NewValue)
 		{
 			if (TSharedPtr<SFontEditorViewport> PinnedWidget = WeakFontPreviewWidget.Pin())
diff --git a/Engine/Source/Editor/GraphEditor/Private/KismetPins/SGraphPinColor.cpp b/Engine/Source/Editor/GraphEditor/Private/KismetPins/SGraphPinColor.cpp
--- a/Engine/Source/Editor/GraphEditor/Private/KismetPins/SGraphPinColor.cpp
+++ b/Engine/Source/Editor/GraphEditor/Private/KismetPins/SGraphPinColor.cpp
@@ -62,7 +62,7 @@ FReply SGraphPinColor::OnColorBoxClicked(const FGeometry& MyGeometry, const FPoi
 		PickerArgs.bIsModal = true;
 		PickerArgs.ParentWidget = DefaultValueWidget;
 		PickerArgs.DisplayGamma = TAttribute<float>::Create(TAttribute<float>::FGetter::CreateUObject(GEngine, &UEngine::GetDisplayGamma));
-		PickerArgs.InitialColorOverride = GetColor();
+		PickerArgs.InitialColor = GetColor();
 		PickerArgs.OnColorCommitted = FOnLinearColorValueChanged::CreateSP(this, &SGraphPinColor::OnColorCommitted);
 		PickerArgs.bUseAlpha = true;
 
diff --git a/Engine/Source/Editor/Kismet/Private/BlueprintDetailsCustomization.cpp b/Engine/Source/Editor/Kismet/Private/BlueprintDetailsCustomization.cpp
--- a/Engine/Source/Editor/Kismet/Private/BlueprintDetailsCustomization.cpp
+++ b/Engine/Source/Editor/Kismet/Private/BlueprintDetailsCustomization.cpp
@@ -5571,7 +5571,7 @@ FReply FBlueprintGraphActionDetails::ColorBlock_OnMouseButtonDown(const FGeometr
 			PickerArgs.bIsModal = true;
 			PickerArgs.ParentWidget = ColorBlock;
 			PickerArgs.DisplayGamma = TAttribute<float>::Create( TAttribute<float>::FGetter::CreateUObject(GEngine, &UEngine::GetDisplayGamma) );
-			PickerArgs.InitialColorOverride = Metadata->InstanceTitleColor;
+			PickerArgs.InitialColor = Metadata->InstanceTitleColor;
 			PickerArgs.OnColorCommitted = FOnLinearColorValueChanged::CreateLambda([WeakSelf](FLinearColor NewValue)
 			{
 				if (TSharedPtr<FBlueprintGraphActionDetails> Self = WeakSelf.Pin())
diff --git a/Engine/Source/Editor/MovieSceneTools/Private/Sections/MovieSceneSectionHelpers.cpp b/Engine/Source/Editor/MovieSceneTools/Private/Sections/MovieSceneSectionHelpers.cpp
--- a/Engine/Source/Editor/MovieSceneTools/Private/Sections/MovieSceneSectionHelpers.cpp
+++ b/Engine/Source/Editor/MovieSceneTools/Private/Sections/MovieSceneSectionHelpers.cpp
@@ -137,7 +137,7 @@ FMovieSceneKeyColorPicker::FMovieSceneKeyColorPicker(UMovieSceneSection* Section
 	FColorPickerArgs PickerArgs;
 	PickerArgs.bUseAlpha = false;
 	PickerArgs.DisplayGamma = TAttribute<float>::Create(TAttribute<float>::FGetter::CreateUObject(GEngine, &UEngine::GetDisplayGamma));
-	PickerArgs.InitialColorOverride = InitialColor;
+	PickerArgs.InitialColor = InitialColor;
 	PickerArgs.ParentWidget = FSlateApplication::Get().GetActiveTopLevelWindow();
 
 	PickerArgs.OnColorCommitted = FOnLinearColorValueChanged::CreateRaw(this, &FMovieSceneKeyColorPicker::OnColorPickerPicked, RChannel, GChannel, BChannel, AChannel);
diff --git a/Engine/Source/Editor/Persona/Private/PoseWatchManagerColumnColor.cpp b/Engine/Source/Editor/Persona/Private/PoseWatchManagerColumnColor.cpp
--- a/Engine/Source/Editor/Persona/Private/PoseWatchManagerColumnColor.cpp
+++ b/Engine/Source/Editor/Persona/Private/PoseWatchManagerColumnColor.cpp
@@ -81,7 +81,7 @@ public:
 			PickerArgs.DisplayGamma = TAttribute<float>::Create(TAttribute<float>::FGetter::CreateUObject(GEngine, &UEngine::GetDisplayGamma));
 			PickerArgs.OnColorCommitted = FOnLinearColorValueChanged::CreateSP(this, &SColorBoxWidget::OnSetColorFromColorPicker);
 			PickerArgs.OnColorPickerCancelled = FOnColorPickerCancelled::CreateSP(this, &SColorBoxWidget::OnColorPickerCancelled);
-			PickerArgs.InitialColorOverride = TreeItem->GetColor();
+			PickerArgs.InitialColor = TreeItem->GetColor();
 			PickerArgs.ParentWidget = ColorPickerParentWidget;
 			PickerArgs.OptionalOwningDetailsView = ColorPickerParentWidget;
 			FWidgetPath ParentWidgetPath;
diff --git a/Engine/Source/Editor/PropertyEditor/Private/SDetailsViewBase.cpp b/Engine/Source/Editor/PropertyEditor/Private/SDetailsViewBase.cpp
--- a/Engine/Source/Editor/PropertyEditor/Private/SDetailsViewBase.cpp
+++ b/Engine/Source/Editor/PropertyEditor/Private/SDetailsViewBase.cpp
@@ -422,6 +422,7 @@ void SDetailsViewBase::CreateColorPickerWindow(const TSharedRef< FPropertyEditor
 	PropertyEditor->GetPropertyNode()->GetReadAddress(false, ReadAddresses, false);
 
 	TOptional<FLinearColor> DefaultColor;
+	bool bClampValue = false;
 	if (ReadAddresses.Num())
 	{
 		for (int32 ColorIndex = 0; ColorIndex < ReadAddresses.Num(); ++ColorIndex)
@@ -432,6 +433,7 @@ void SDetailsViewBase::CreateColorPickerWindow(const TSharedRef< FPropertyEditor
 				if (CastField<FStructProperty>(Property)->Struct->GetFName() == NAME_Color)
 				{
 					DefaultColor = *reinterpret_cast<const FColor*>(Addr);
+					bClampValue = true;
 				}
 				else
 				{
@@ -451,6 +453,7 @@ void SDetailsViewBase::CreateColorPickerWindow(const TSharedRef< FPropertyEditor
 		FColorPickerArgs PickerArgs = FColorPickerArgs(DefaultColor.GetValue(), FOnLinearColorValueChanged::CreateStatic(&UE::PropertyEditor::Private::SDetailsViewBase_SetColor, WeakPropertyHandle));
 		PickerArgs.ParentWidget = AsShared();
 		PickerArgs.bUseAlpha = bUseAlpha;
+		PickerArgs.bClampValue = bClampValue;
 		PickerArgs.DisplayGamma = TAttribute<float>::Create(TAttribute<float>::FGetter::CreateUObject(GEngine, &UEngine::GetDisplayGamma));
 		PickerArgs.OnColorPickerWindowClosed = FOnWindowClosed::CreateSP(this, &SDetailsViewBase::OnColorPickerWindowClosed);
 		PickerArgs.OptionalOwningDetailsView = AsShared();
diff --git a/Engine/Source/Editor/PropertyEditor/Private/SPropertyTreeViewImpl.cpp b/Engine/Source/Editor/PropertyEditor/Private/SPropertyTreeViewImpl.cpp
--- a/Engine/Source/Editor/PropertyEditor/Private/SPropertyTreeViewImpl.cpp
+++ b/Engine/Source/Editor/PropertyEditor/Private/SPropertyTreeViewImpl.cpp
@@ -1230,6 +1230,7 @@ void SPropertyTreeViewImpl::CreateColorPickerWindow(const TSharedRef< class FPro
 
 	// Use the first address for the initial color
 	TOptional<FLinearColor> DefaultColor;
+	bool bClampValue = false;
 	if (ReadAddresses.Num())
 	{
 		const uint8* Addr = ReadAddresses.GetAddress(0);
@@ -1238,6 +1239,7 @@ void SPropertyTreeViewImpl::CreateColorPickerWindow(const TSharedRef< class FPro
 			if (CastField<FStructProperty>(Property)->Struct->GetFName() == NAME_Color)
 			{
 				DefaultColor = *reinterpret_cast<const FColor*>(Addr);
+				bClampValue = true;
 			}
 			else
 			{
@@ -1253,6 +1255,7 @@ void SPropertyTreeViewImpl::CreateColorPickerWindow(const TSharedRef< class FPro
 		FColorPickerArgs PickerArgs = FColorPickerArgs(DefaultColor.GetValue(), FOnLinearColorValueChanged::CreateStatic(&UE::PropertyEditor::Private::CreateColorPickerWindow_SetColor, WeakPropertyHandle));
 		PickerArgs.ParentWidget = AsShared();
 		PickerArgs.bUseAlpha = bUseAlpha;
+		PickerArgs.bClampValue = bClampValue;
 		PickerArgs.DisplayGamma = TAttribute<float>::Create(TAttribute<float>::FGetter::CreateUObject(GEngine, &UEngine::GetDisplayGamma));
 
 		OpenColorPicker(PickerArgs);
diff --git a/Engine/Source/Editor/PropertyEditor/Private/SSingleProperty.cpp b/Engine/Source/Editor/PropertyEditor/Private/SSingleProperty.cpp
--- a/Engine/Source/Editor/PropertyEditor/Private/SSingleProperty.cpp
+++ b/Engine/Source/Editor/PropertyEditor/Private/SSingleProperty.cpp
@@ -320,6 +320,7 @@ void SSingleProperty::CreateColorPickerWindow( const TSharedRef< class FProperty
 
 		// Use the first address for the initial color
 		TOptional<FLinearColor> DefaultColor;
+		bool bClampValue = false;
 		if( ReadAddresses.Num() ) 
 		{
 			const uint8* Addr = ReadAddresses.GetAddress(0);
@@ -328,6 +329,7 @@ void SSingleProperty::CreateColorPickerWindow( const TSharedRef< class FProperty
 				if( CastField<FStructProperty>(Property)->Struct->GetFName() == NAME_Color )
 				{
 					DefaultColor = *reinterpret_cast<const FColor*>(Addr);
+					bClampValue = true;
 				}
 				else
 				{
@@ -342,6 +344,7 @@ void SSingleProperty::CreateColorPickerWindow( const TSharedRef< class FProperty
 			FColorPickerArgs PickerArgs = FColorPickerArgs(DefaultColor.GetValue(), FOnLinearColorValueChanged::CreateSP(this, &SSingleProperty::SetColorPropertyFromColorPicker));
 			PickerArgs.ParentWidget = AsShared();
 			PickerArgs.bUseAlpha = bUseAlpha;
+			PickerArgs.bClampValue = bClampValue;
 			PickerArgs.DisplayGamma = TAttribute<float>::Create(TAttribute<float>::FGetter::CreateUObject(GEngine, &UEngine::GetDisplayGamma));
 
 			OpenColorPicker(PickerArgs);
diff --git a/Engine/Source/Editor/PropertyEditor/Private/UserInterface/PropertyEditor/SPropertyEditorColor.cpp b/Engine/Source/Editor/PropertyEditor/Private/UserInterface/PropertyEditor/SPropertyEditorColor.cpp
--- a/Engine/Source/Editor/PropertyEditor/Private/UserInterface/PropertyEditor/SPropertyEditorColor.cpp
+++ b/Engine/Source/Editor/PropertyEditor/Private/UserInterface/PropertyEditor/SPropertyEditorColor.cpp
@@ -165,6 +165,7 @@ void SPropertyEditorColor::CreateColorPickerWindow(const TSharedRef< class FProp
 		OriginalColors.Empty( ReadAddresses.Num() );
 		OriginalColors.AddUninitialized( ReadAddresses.Num() );
 
+		bool bClampValue = false;
 		// Store off the original colors in the case that the user cancels the color picker. We'll revert to the original colors in that case
 		for( int32 AddrIndex = 0; AddrIndex <  ReadAddresses.Num(); ++AddrIndex )
 		{
@@ -175,6 +176,7 @@ void SPropertyEditorColor::CreateColorPickerWindow(const TSharedRef< class FProp
 				if( CastField<FStructProperty>(Property)->Struct->GetFName() == NAME_Color )
 				{
 					OriginalColors[AddrIndex] = ((FColor*)Addr)->ReinterpretAsLinear();
+					bClampValue = true;
 				}
 				else
 				{
@@ -192,10 +194,11 @@ void SPropertyEditorColor::CreateColorPickerWindow(const TSharedRef< class FProp
 		PickerArgs.ParentWidget = AsShared();
 		PickerArgs.bUseAlpha = bUseAlpha;
 		PickerArgs.bOnlyRefreshOnOk = bOnlyRefreshOnOk;
+		PickerArgs.bClampValue = bClampValue;
 		PickerArgs.DisplayGamma = TAttribute<float>::Create( TAttribute<float>::FGetter::CreateUObject(GEngine, &UEngine::GetDisplayGamma) );
 		PickerArgs.OnColorCommitted = FOnLinearColorValueChanged::CreateSP( this, &SPropertyEditorColor::SetColor);
 		PickerArgs.OnColorPickerCancelled = FOnColorPickerCancelled::CreateSP( this, &SPropertyEditorColor::OnColorPickerCancelled );
-		PickerArgs.InitialColorOverride = InitialColor;
+		PickerArgs.InitialColor = InitialColor;
 
 		OpenColorPicker(PickerArgs);
 	}
diff --git a/Engine/Source/Editor/Sequencer/Private/MVVM/ViewModels/FolderModel.cpp b/Engine/Source/Editor/Sequencer/Private/MVVM/ViewModels/FolderModel.cpp
--- a/Engine/Source/Editor/Sequencer/Private/MVVM/ViewModels/FolderModel.cpp
+++ b/Engine/Source/Editor/Sequencer/Private/MVVM/ViewModels/FolderModel.cpp
@@ -359,7 +359,7 @@ void FFolderModel::SetFolderColor()
 	FColorPickerArgs PickerArgs;
 	PickerArgs.bUseAlpha = false;
 	PickerArgs.DisplayGamma = TAttribute<float>::Create( TAttribute<float>::FGetter::CreateUObject(GEngine, &UEngine::GetDisplayGamma) );
-	PickerArgs.InitialColorOverride = InitialFolderColor.ReinterpretAsLinear();
+	PickerArgs.InitialColor = InitialFolderColor.ReinterpretAsLinear();
 	PickerArgs.OnColorCommitted = FOnLinearColorValueChanged::CreateSP(this, &FFolderModel::OnColorPickerPicked);
 	PickerArgs.OnColorPickerWindowClosed = FOnWindowClosed::CreateSP(this, &FFolderModel::OnColorPickerClosed);
 	PickerArgs.OnColorPickerCancelled  = FOnColorPickerCancelled::CreateSP(this, &FFolderModel::OnColorPickerCancelled );
diff --git a/Engine/Source/Editor/Sequencer/Private/MVVM/Views/SOutlinerTrackView.cpp b/Engine/Source/Editor/Sequencer/Private/MVVM/Views/SOutlinerTrackView.cpp
--- a/Engine/Source/Editor/Sequencer/Private/MVVM/Views/SOutlinerTrackView.cpp
+++ b/Engine/Source/Editor/Sequencer/Private/MVVM/Views/SOutlinerTrackView.cpp
@@ -207,7 +207,7 @@ FReply SOutlinerTrackView::OnSetTrackColor()
 	FColorPickerArgs PickerArgs;
 	PickerArgs.bUseAlpha = false;
 	PickerArgs.DisplayGamma = TAttribute<float>::Create(TAttribute<float>::FGetter::CreateUObject(GEngine, &UEngine::GetDisplayGamma));
-	PickerArgs.InitialColorOverride = AnimationOutlinerTreeNode::InitialTrackColor;
+	PickerArgs.InitialColor = AnimationOutlinerTreeNode::InitialTrackColor;
 	PickerArgs.ParentWidget = GetParentWidget();
 	PickerArgs.OnColorCommitted = FOnLinearColorValueChanged::CreateSP(this, &SOutlinerTrackView::OnColorPickerPicked);
 	PickerArgs.OnColorPickerWindowClosed = FOnWindowClosed::CreateSP(this, &SOutlinerTrackView::OnColorPickerClosed);
diff --git a/Engine/Source/Editor/UnrealEd/Private/SColorGradientEditor.cpp b/Engine/Source/Editor/UnrealEd/Private/SColorGradientEditor.cpp
--- a/Engine/Source/Editor/UnrealEd/Private/SColorGradientEditor.cpp
+++ b/Engine/Source/Editor/UnrealEd/Private/SColorGradientEditor.cpp
@@ -655,7 +655,7 @@ void SColorGradientEditor::OpenGradientStopColorPicker()
 		ColorPickerArgs.bIsModal = false;
 		ColorPickerArgs.ParentWidget = SharedThis( this );
 		ColorPickerArgs.bUseAlpha = false;
-		ColorPickerArgs.InitialColorOverride = SelectedStop.GetColor( *CurveOwner );
+		ColorPickerArgs.InitialColor = SelectedStop.GetColor( *CurveOwner );
 		ColorPickerArgs.OnColorCommitted = FOnLinearColorValueChanged::CreateSP( this, &SColorGradientEditor::OnSelectedStopColorChanged );
 		ColorPickerArgs.OnColorPickerCancelled  = FOnColorPickerCancelled::CreateSP( this, &SColorGradientEditor::OnCancelSelectedStopColorChange );
 		OpenColorPicker( ColorPickerArgs );
diff --git a/Engine/Source/Editor/WorldBrowser/Private/SWorldHierarchyItem.cpp b/Engine/Source/Editor/WorldBrowser/Private/SWorldHierarchyItem.cpp
--- a/Engine/Source/Editor/WorldBrowser/Private/SWorldHierarchyItem.cpp
+++ b/Engine/Source/Editor/WorldBrowser/Private/SWorldHierarchyItem.cpp
@@ -476,7 +476,7 @@ FReply SWorldHierarchyItem::OnChangeColor()
 	{
 		FColorPickerArgs PickerArgs;
 		PickerArgs.DisplayGamma = TAttribute<float>::Create(TAttribute<float>::FGetter::CreateUObject(GEngine, &UEngine::GetDisplayGamma));
-		PickerArgs.InitialColorOverride = WorldTreeItem->GetDrawColor();
+		PickerArgs.InitialColor = WorldTreeItem->GetDrawColor();
 		PickerArgs.bOnlyRefreshOnMouseUp = false;
 		PickerArgs.bOnlyRefreshOnOk = false;
 		PickerArgs.OnColorCommitted = FOnLinearColorValueChanged::CreateSP(this, &SWorldHierarchyItem::OnSetColorFromColorPicker);
diff --git a/Engine/Source/Runtime/AppFramework/Private/Widgets/Colors/SColorPicker.cpp b/Engine/Source/Runtime/AppFramework/Private/Widgets/Colors/SColorPicker.cpp
--- a/Engine/Source/Runtime/AppFramework/Private/Widgets/Colors/SColorPicker.cpp
+++ b/Engine/Source/Runtime/AppFramework/Private/Widgets/Colors/SColorPicker.cpp
@@ -63,14 +63,10 @@ void SColorPicker::Construct( const FArguments& InArgs )
 	CurrentColorHSV = OldColor = TargetColorAttribute.Get().LinearRGBToHSV();
 	CurrentColorRGB = TargetColorAttribute.Get();
 	CurrentMode = EColorPickerModes::Wheel;
-	TargetFColors = InArgs._TargetFColors.Get();
-	TargetLinearColors = InArgs._TargetLinearColors.Get();
-	TargetColorChannels = InArgs._TargetColorChannels.Get();
 	bUseAlpha = InArgs._UseAlpha;
 	bOnlyRefreshOnMouseUp = InArgs._OnlyRefreshOnMouseUp.Get();
 	bOnlyRefreshOnOk = InArgs._OnlyRefreshOnOk.Get();
 	OnColorCommitted = InArgs._OnColorCommitted;
-	PreColorCommitted = InArgs._PreColorCommitted;
 	OnColorPickerCancelled = InArgs._OnColorPickerCancelled;
 	OnInteractivePickBegin = InArgs._OnInteractivePickBegin;
 	OnInteractivePickEnd = InArgs._OnInteractivePickEnd;
@@ -79,6 +75,7 @@ void SColorPicker::Construct( const FArguments& InArgs )
 	DisplayGamma = InArgs._DisplayGamma;
 	bClosedViaOkOrCancel = false;
 	bValidCreationOverrideExists = InArgs._OverrideColorPickerCreation;
+	bClampValue = InArgs._ClampValue;
 	OptionalOwningDetailsView = InArgs._OptionalOwningDetailsView.Get().IsValid() ? InArgs._OptionalOwningDetailsView.Get() : nullptr;
 
 	if ( InArgs._sRGBOverride.IsSet() )
@@ -98,9 +95,6 @@ void SColorPicker::Construct( const FArguments& InArgs )
 	bColorPickerIsInlineVersion = InArgs._DisplayInlineVersion;
 	bIsInteractive = false;
 	bPerfIsTooSlowToUpdate = false;
-	
-
-	BackupColors();
 
 	BeginAnimation(FLinearColor(ForceInit), CurrentColorHSV);
 
@@ -133,109 +127,6 @@ void SColorPicker::Construct( const FArguments& InArgs )
 /* SColorPicker implementation
  *****************************************************************************/
 
-void SColorPicker::BackupColors()
-{
-	OldTargetFColors.Empty();
-	for (int32 i = 0; i < TargetFColors.Num(); ++i)
-	{
-		OldTargetFColors.Add( *TargetFColors[i] );
-	}
-
-	OldTargetLinearColors.Empty();
-	for (int32 i = 0; i < TargetLinearColors.Num(); ++i)
-	{
-		OldTargetLinearColors.Add( *TargetLinearColors[i] );
-	}
-
-	OldTargetColorChannels.Empty();
-	for (int32 i = 0; i < TargetColorChannels.Num(); ++i)
-	{
-		// Remap the color channel as a linear color for ease
-		const FColorChannels& Channel = TargetColorChannels[i];
-		const FLinearColor Color( Channel.Red ? *Channel.Red : 0.f, Channel.Green ? *Channel.Green : 0.f, Channel.Blue ? *Channel.Blue : 0.f, Channel.Alpha ? *Channel.Alpha : 0.f );
-		OldTargetColorChannels.Add( Color );
-	}
-}
-
-
-void SColorPicker::RestoreColors()
-{
-	check(TargetFColors.Num() == OldTargetFColors.Num());
-
-	for (int32 i = 0; i < TargetFColors.Num(); ++i)
-	{
-		*TargetFColors[i] = OldTargetFColors[i];
-	}
-
-	check(TargetLinearColors.Num() == OldTargetLinearColors.Num());
-
-	for (int32 i = 0; i < TargetLinearColors.Num(); ++i)
-	{
-		*TargetLinearColors[i] = OldTargetLinearColors[i];
-	}
-
-	check(TargetColorChannels.Num() == OldTargetColorChannels.Num());
-
-	for (int32 i = 0; i < TargetColorChannels.Num(); ++i)
-	{
-		// Copy back out of the linear to the color channel
-		FColorChannels& Channel = TargetColorChannels[i];
-		const FLinearColor& OldChannel = OldTargetColorChannels[i];
-		if (Channel.Red)
-		{
-			*Channel.Red = OldChannel.R;
-		}
-		if (Channel.Green)
-		{
-			*Channel.Green = OldChannel.G;
-		}
-		if (Channel.Blue)
-		{
-			*Channel.Blue = OldChannel.B;
-		}
-		if (Channel.Alpha)
-		{
-			*Channel.Alpha = OldChannel.A;
-		}
-	}
-}
-
-
-void SColorPicker::SetColors(const FLinearColor& InColor)
-{
-	for (int32 i = 0; i < TargetFColors.Num(); ++i)
-	{
-		*TargetFColors[i] = InColor.ToFColor(true);
-	}
-
-	for (int32 i = 0; i < TargetLinearColors.Num(); ++i)
-	{
-		*TargetLinearColors[i] = InColor;
-	}
-
-	for (int32 i = 0; i < TargetColorChannels.Num(); ++i)
-	{
-		// Only set those channels who have a valid ptr
-		FColorChannels& Channel = TargetColorChannels[i];
-		if (Channel.Red)
-		{
-			*Channel.Red = InColor.R;
-		}
-		if (Channel.Green)
-		{
-			*Channel.Green = InColor.G;
-		}
-		if (Channel.Blue)
-		{
-			*Channel.Blue = InColor.B;
-		}
-		if (Channel.Alpha)
-		{
-			*Channel.Alpha = InColor.A;
-		}
-	}
-}
-
 
 BEGIN_SLATE_FUNCTION_BUILD_OPTIMIZATION
 void SColorPicker::GenerateDefaultColorPickerContent( bool bAdvancedSectionExpanded )
@@ -727,7 +618,6 @@ void SColorPicker::DiscardColor()
 	else
 	{	
 		SetNewTargetColorHSV(OldColor, true);
-		RestoreColors();
 	}
 }
 
@@ -786,9 +676,6 @@ void SColorPicker::UpdateColorPick()
 	bPerfIsTooSlowToUpdate = false;
 	FLinearColor OutColor = CurrentColorRGB;
 
-	PreColorCommitted.ExecuteIfBound(OutColor);
-
-	SetColors(OutColor);
 	OnColorCommitted.ExecuteIfBound(OutColor);
 	
 	// This callback is only necessary for wx backwards compatibility
@@ -911,7 +798,7 @@ TSharedRef<SWidget> SColorPicker::MakeColorSpinBox( EColorPickerChannels Channel
 	}
 
 	const int32 GradientHeight = 6;
-	const float HDRMaxValue = (TargetFColors.Num()) ? 1.f : FLT_MAX;
+	const float HDRMaxValue = bClampValue ? 1.f : FLT_MAX;
 	const FSlateFontInfo SmallLayoutFont = FAppStyle::Get().GetFontStyle("ColorPicker.Font");
 
 	// create gradient widget
@@ -1336,7 +1223,6 @@ void SColorPicker::HandleEyeDropperButtonComplete(bool bCancelled)
 	if (bCancelled)
 	{
 		SetNewTargetColorHSV(OldColor, true);
-		RestoreColors();
 	}
 
 	if (bOnlyRefreshOnMouseUp || bPerfIsTooSlowToUpdate)
@@ -1653,27 +1539,8 @@ bool OpenColorPicker(const FColorPickerArgs& Args)
 
 	// Consoles do not support opening new windows
 #if PLATFORM_DESKTOP
-	FLinearColor OldColor = Args.InitialColorOverride;
-
-	if (Args.ColorArray && Args.ColorArray->Num() > 0)
-	{
-		OldColor = FLinearColor(*(*Args.ColorArray)[0]);
-	}
-	else if (Args.LinearColorArray && Args.LinearColorArray->Num() > 0)
-	{
-		OldColor = *(*Args.LinearColorArray)[0];
-	}
-	else if (Args.ColorChannelsArray && Args.ColorChannelsArray->Num() > 0)
-	{
-		OldColor.R = (*Args.ColorChannelsArray)[0].Red ? *(*Args.ColorChannelsArray)[0].Red : 0.0f;
-		OldColor.G = (*Args.ColorChannelsArray)[0].Green ? *(*Args.ColorChannelsArray)[0].Green : 0.0f;
-		OldColor.B = (*Args.ColorChannelsArray)[0].Blue ? *(*Args.ColorChannelsArray)[0].Blue : 0.0f;
-		OldColor.A = (*Args.ColorChannelsArray)[0].Alpha ? *(*Args.ColorChannelsArray)[0].Alpha : 0.0f;
-	}
-	else
-	{
-		check(Args.OnColorCommitted.IsBound());
-	}
+	FLinearColor OldColor = Args.InitialColor;
+	ensureMsgf(Args.OnColorCommitted.IsBound(), TEXT("OnColorCommitted should be bound to set the color."));
 		
 	// Determine the position of the window so that it will spawn near the mouse, but not go off the screen.
 	FVector2D CursorPos = FSlateApplication::Get().GetCursorPos();
@@ -1726,15 +1593,11 @@ bool OpenColorPicker(const FColorPickerArgs& Args)
 
 	TSharedRef<SColorPicker> CreatedColorPicker = SNew(SColorPicker)
 		.TargetColorAttribute(OldColor)
-		.TargetFColors(Args.ColorArray ? *Args.ColorArray : TArray<FColor*>())
-		.TargetLinearColors(Args.LinearColorArray ? *Args.LinearColorArray : TArray<FLinearColor*>())
-		.TargetColorChannels(Args.ColorChannelsArray ? *Args.ColorChannelsArray : TArray<FColorChannels>())
 		.UseAlpha(Args.bUseAlpha)
 		.ExpandAdvancedSection(Args.bExpandAdvancedSection)
 		.OnlyRefreshOnMouseUp(Args.bOnlyRefreshOnMouseUp && !Args.bIsModal)
 		.OnlyRefreshOnOk(Args.bOnlyRefreshOnOk || Args.bIsModal)
 		.OnColorCommitted(Args.OnColorCommitted)
-		.PreColorCommitted(Args.PreColorCommitted)
 		.OnColorPickerCancelled(Args.OnColorPickerCancelled)
 		.OnInteractivePickBegin(Args.OnInteractivePickBegin)
 		.OnInteractivePickEnd(Args.OnInteractivePickEnd)
diff --git a/Engine/Source/Runtime/AppFramework/Private/Widgets/Testing/STestSuite.cpp b/Engine/Source/Runtime/AppFramework/Private/Widgets/Testing/STestSuite.cpp
--- a/Engine/Source/Runtime/AppFramework/Private/Widgets/Testing/STestSuite.cpp
+++ b/Engine/Source/Runtime/AppFramework/Private/Widgets/Testing/STestSuite.cpp
@@ -3463,7 +3463,7 @@ protected:
 		PickerArgs.bOnlyRefreshOnOk = false;
 		PickerArgs.OnColorCommitted = FOnLinearColorValueChanged::CreateSP(this, &SRichTextEditTest::SetFontColor);
 		PickerArgs.OnColorPickerCancelled = FOnColorPickerCancelled::CreateSP(this, &SRichTextEditTest::SetFontColor);
-		PickerArgs.InitialColorOverride = FontColor;
+		PickerArgs.InitialColor = FontColor;
 
 		OpenColorPicker(PickerArgs);
 
diff --git a/Engine/Source/Runtime/AppFramework/Public/Widgets/Colors/SColorPicker.h b/Engine/Source/Runtime/AppFramework/Public/Widgets/Colors/SColorPicker.h
--- a/Engine/Source/Runtime/AppFramework/Public/Widgets/Colors/SColorPicker.h
+++ b/Engine/Source/Runtime/AppFramework/Public/Widgets/Colors/SColorPicker.h
@@ -64,7 +64,6 @@ struct FColorChannels
 	float* Alpha;
 };
 
-
 /**
  * Class for placing a color picker. If all you need is a standalone color picker,
  * use the functions OpenColorPicker and DestroyColorPicker, since they hold a static
@@ -77,9 +76,6 @@ public:
 
 	SLATE_BEGIN_ARGS(SColorPicker)
 		: _TargetColorAttribute(FLinearColor(ForceInit))
-		, _TargetFColors()
-		, _TargetLinearColors()
-		, _TargetColorChannels()
 		, _UseAlpha(true)
 		, _OnlyRefreshOnMouseUp(false)
 		, _OnlyRefreshOnOk(false)
@@ -95,39 +91,41 @@ public:
 		, _DisplayInlineVersion(false)
 		, _OverrideColorPickerCreation(false)
 		, _ExpandAdvancedSection(false)
+		, _ClampValue(false)
 		, _OptionalOwningDetailsView(nullptr)
 	{ }
-		
+
 		/** The color that is being targeted as a TAttribute */
 		SLATE_ATTRIBUTE(FLinearColor, TargetColorAttribute)
-		
+
 		/** An array of color pointers this color picker targets */
-		SLATE_ATTRIBUTE(TArray<FColor*>, TargetFColors)
+		SLATE_ARGUMENT_DEPRECATED(TArray<FColor*>, TargetFColors, 5.2, "TargetFColors is deprecated. Use OnColorCommitted to get the selected color.")
 		
 		/** An array of linear color pointers this color picker targets */
-		SLATE_ATTRIBUTE(TArray<FLinearColor*>, TargetLinearColors)
-		
+		SLATE_ARGUMENT_DEPRECATED(TArray<FLinearColor*>, TargetLinearColors, 5.2, "TargetLinearColors is deprecated. Use OnColorCommitted to get the selected color.")
+
 		/**
 		 * An array of color pointer structs this color picker targets
 		 * Only to keep compatibility with wx. Should be removed once wx is gone.
 		 */
-		SLATE_ATTRIBUTE(TArray<FColorChannels>, TargetColorChannels)
+		SLATE_ARGUMENT_DEPRECATED(TArray<FColorChannels>, TargetColorChannels, 5.2, "TargetColorChannels is deprecated. Use OnColorCommitted to get the selected color.")
 
 		/** Whether the ability to pick the alpha value is enabled */
 		SLATE_ATTRIBUTE(bool, UseAlpha)
-		
+
 		/** Prevents immediate refreshs for performance reasons. */
 		SLATE_ATTRIBUTE(bool, OnlyRefreshOnMouseUp)
 
 		/** Prevents multiple refreshes when requested. */
 		SLATE_ATTRIBUTE(bool, OnlyRefreshOnOk)
-		
+
 		/** The event called when the color is committed */
 		SLATE_EVENT(FOnLinearColorValueChanged, OnColorCommitted)
 
+		UE_DEPRECATED(5.2, "PreColorCommitted is deprecated. Use OnColorCommitted to update your values.")
 		/** The event called before the color is committed */
 		SLATE_EVENT(FOnLinearColorValueChanged, PreColorCommitted)
-		
+
 		/** The event called when the color picker cancel button is pressed */
 		SLATE_EVENT(FOnColorPickerCancelled, OnColorPickerCancelled)
 
@@ -158,6 +156,9 @@ public:
 		/** If true, the Advanced section will be expanded, regardless of the remembered state */
 		SLATE_ARGUMENT(bool, ExpandAdvancedSection)
 
+		/** The LinearColor is expected to be converted to a FColor, set the clamp channel value. */
+		SLATE_ARGUMENT(bool, ClampValue)
+
 		/** Allows a details view to own the color picker so refreshing another details view doesn't close it */
 		SLATE_ATTRIBUTE(TSharedPtr<SWidget>, OptionalOwningDetailsView)
 
@@ -204,12 +205,6 @@ protected:
 	/** Backup all the colors that are being modified */
 	void BackupColors();
 
-	/** Restore all the modified colors to their original state */
-	void RestoreColors();
-
-	/** Set all the colors to this new color */
-	void SetColors(const FLinearColor& InColor);
-
 	bool ApplyNewTargetColor(bool bForceUpdate = false);
 
 	void GenerateDefaultColorPickerContent(bool bAdvancedSectionExpanded);
@@ -403,27 +398,6 @@ private:
 	/** If true, then the performance is too bad to have auto-updating */
 	bool bPerfIsTooSlowToUpdate;
 
-	/** An array of color pointers this color picker targets */
-	TArray<FColor*> TargetFColors;
-	
-	/** An array of linear color pointers this color picker targets */
-	TArray<FLinearColor*> TargetLinearColors;
-	
-	/**
-	 * An array of color pointer structs this color picker targets
-	 * Only to keep compatibility with wx. Should be removed once wx is gone.
-	 */
-	TArray<FColorChannels> TargetColorChannels;
-
-	/** Backups of the TargetFColors */
-	TArray<FColor> OldTargetFColors;
-
-	/** Backups of the TargetLinearColors */
-	TArray<FLinearColor> OldTargetLinearColors;
-
-	/** Backups of the TargetColorChannels */
-	TArray<FLinearColor> OldTargetColorChannels;
-
 	/** Whether or not the color uses Alpha or not */
 	TAttribute<bool> bUseAlpha;
 
@@ -468,15 +442,15 @@ private:
 
 	/** Is true if the color picker creation behavior can be overridden */
 	bool bValidCreationOverrideExists;
+	
+	/** The current display gamma used to correct colors picked from the display. */
+	bool bClampValue;
 
 private:
 
 	/** Invoked when a new value is selected on the color wheel */
 	FOnLinearColorValueChanged OnColorCommitted;
 
-	/** Invoked before a new value is selected on the color wheel */
-	FOnLinearColorValueChanged PreColorCommitted;
-
 	/** Invoked when the color picker cancel button is pressed */
 	FOnColorPickerCancelled OnColorPickerCancelled;
 
@@ -522,24 +496,31 @@ struct FColorPickerArgs
 	/** Whether to open the color picker as a menu window. */
 	bool bOpenAsMenu = false;
 
+	/** Set true if values should be max to 1.0. HDR values may go over 1.0. */
+	bool bClampValue = false;
+
 	/** The current display gamma used to correct colors picked from the display. */
 	TAttribute<float> DisplayGamma = 2.2f;
 
 	/** If set overrides the global option for the desired setting of sRGB mode. */
 	TOptional<bool> sRGBOverride = TOptional<bool>();
 
+	UE_DEPRECATED(5.2, "ColorArray is deprecated. Use OnColorCommitted to update your values.")
 	/** An array of FColors to target. */
 	const TArray<FColor*>* ColorArray = nullptr;
 
+	UE_DEPRECATED(5.2, "LinearColorArray is deprecated. Use OnColorCommitted to update your values.")
 	/** An array of FLinearColors to target. */
 	const TArray<FLinearColor*>* LinearColorArray = nullptr;
 
+	UE_DEPRECATED(5.2, "ColorChannelsArray is deprecated. Use OnColorCommitted to update your values.")
 	/** An array of FColorChannels to target. (deprecated now that wx is gone?) */
 	const TArray<FColorChannels>* ColorChannelsArray = nullptr;
 
 	/** A delegate to be called when the color changes. */
 	FOnLinearColorValueChanged OnColorCommitted;
 
+	UE_DEPRECATED(5.2, "PreColorCommitted is deprecated. Use OnColorCommitted to update your values.")
 	/** A delegate to be called before the color change is committed. */
 	FOnLinearColorValueChanged PreColorCommitted;
 
@@ -555,9 +536,13 @@ struct FColorPickerArgs
 	/** A delegate to be called when a slider drag, color wheel drag or dropper grab finishes */
 	FSimpleDelegate OnInteractivePickEnd;
 
+	UE_DEPRECATED(5.2, "InitialColorOverride is deprecated. Use InitialColor to set the initial color.")
 	/** Overrides the initial color set on the color picker. */
 	FLinearColor InitialColorOverride = FLinearColor::White;
 
+	/** The initial color set on the color picker. */
+	FLinearColor InitialColor = FLinearColor::White;
+
 	/** Allows a details view to own the color picker so refreshing another details view doesn't close it */
 	TSharedPtr<SWidget> OptionalOwningDetailsView;
 
@@ -566,8 +551,13 @@ struct FColorPickerArgs
 
 	FColorPickerArgs(FLinearColor InInitialColor, FOnLinearColorValueChanged InOnColorCommitted)
 		: OnColorCommitted(MoveTemp(InOnColorCommitted))
-		, InitialColorOverride(InInitialColor)
+		, InitialColor(InInitialColor)
 	{}
+
+PRAGMA_DISABLE_DEPRECATION_WARNINGS
+	FColorPickerArgs(const FColorPickerArgs&) = default;
+	FColorPickerArgs(FColorPickerArgs&&) = default;
+PRAGMA_ENABLE_DEPRECATION_WARNINGS
 };
 
 /** Get a pointer to the static color picker, or nullptr if it does not exist. */
