Resize the hex color input in SColorPicker

diff --git a/Engine/Source/Runtime/AppFramework/Private/Widgets/Colors/SColorPicker.cpp b/Engine/Source/Runtime/AppFramework/Private/Widgets/Colors/SColorPicker.cpp
--- a/Engine/Source/Runtime/AppFramework/Private/Widgets/Colors/SColorPicker.cpp
+++ b/Engine/Source/Runtime/AppFramework/Private/Widgets/Colors/SColorPicker.cpp
@@ -446,10 +446,10 @@ void SColorPicker::GenerateDefaultColorPickerContent( bool bAdvancedSectionExpan
 
 										+ SHorizontalBox::Slot()
 											.AutoWidth()
-											.MaxWidth(80.0f)
+											.MaxWidth(90.0f)
 											[
 												SNew(SEditableTextBox)
-													.MinDesiredWidth(80.0f)
+													.MinDesiredWidth(90.0f)
 													.Text(this, &SColorPicker::HandleHexLinearBoxText)
 													.OnTextCommitted(this, &SColorPicker::HandleHexLinearInputTextCommitted)
 											]
@@ -475,10 +475,10 @@ void SColorPicker::GenerateDefaultColorPickerContent( bool bAdvancedSectionExpan
 
 										+ SHorizontalBox::Slot()
 											.AutoWidth()
-											.MaxWidth(80.0f)
+											.MaxWidth(90.0f)
 											[
 												SNew(SEditableTextBox)
-												.MinDesiredWidth(80.0f)
+												.MinDesiredWidth(90.0f)
 												.Text(this, &SColorPicker::HandleHexSRGBBoxText)
 												.OnTextCommitted(this, &SColorPicker::HandleHexSRGBInputTextCommitted)
 											]
