Set the color to full brightness/opacity when picking from the color wheel/spectrum if it was previously almost black or transparent

diff --git a/Engine/Source/Runtime/AppFramework/Private/Widgets/Colors/SColorPicker.cpp b/Engine/Source/Runtime/AppFramework/Private/Widgets/Colors/SColorPicker.cpp
--- a/Engine/Source/Runtime/AppFramework/Private/Widgets/Colors/SColorPicker.cpp
+++ b/Engine/Source/Runtime/AppFramework/Private/Widgets/Colors/SColorPicker.cpp
@@ -1158,6 +1158,15 @@ FLinearColor SColorPicker::HandleColorSliderStartColor( EColorPickerChannels Cha
 
 void SColorPicker::HandleColorSpectrumValueChanged( FLinearColor NewValue )
 {
+	// In this color, R = H, G = S, B = V
+	if (FMath::IsNearlyZero(NewValue.B))
+	{
+		NewValue.B = 1.0f;
+	}
+	if (!bUseAlpha.Get() || FMath::IsNearlyZero(NewValue.A))
+	{
+		NewValue.A = 1.0f;
+	}
 	SetNewTargetColorHSV(NewValue);
 }
 
