Slate: Add ObservableArray pattern to ListView, TitleView and TreeView.The TreeView implementation only supports the new pattern for the root items. It is where we would have the most gain in performance when it is changed.
The observable array automatically notify the view when an item is added/removed from the array.

diff --git a/Engine/Plugins/Editor/ConsoleVariablesEditor/Source/ConsoleVariablesEditor/Private/Views/List/SConsoleVariablesEditorList.h b/Engine/Plugins/Editor/ConsoleVariablesEditor/Source/ConsoleVariablesEditor/Private/Views/List/SConsoleVariablesEditorList.h
--- a/Engine/Plugins/Editor/ConsoleVariablesEditor/Source/ConsoleVariablesEditor/Private/Views/List/SConsoleVariablesEditorList.h
+++ b/Engine/Plugins/Editor/ConsoleVariablesEditor/Source/ConsoleVariablesEditor/Private/Views/List/SConsoleVariablesEditorList.h
@@ -64,11 +64,6 @@ public:
 	[[nodiscard]] TArray<FConsoleVariablesEditorListRowPtr> GetTreeViewItems() const;
 	void SetTreeViewItems(const TArray<FConsoleVariablesEditorListRowPtr>& InItems);
 
-	[[nodiscard]] int32 GetTreeViewItemCount() const
-	{
-		return TreeViewRootObjects.Num();
-	}
-
 	/** Updates the saved values in a UConsoleVariablesAsset so that the command/value map can be saved to disk */
 	void UpdatePresetValuesForSave(const TObjectPtr<UConsoleVariablesAsset> InAsset) const;
 
diff --git a/Engine/Plugins/Editor/ObjectMixer/ObjectMixer/Source/ObjectMixer/Private/Views/List/SObjectMixerEditorList.cpp b/Engine/Plugins/Editor/ObjectMixer/ObjectMixer/Source/ObjectMixer/Private/Views/List/SObjectMixerEditorList.cpp
--- a/Engine/Plugins/Editor/ObjectMixer/ObjectMixer/Source/ObjectMixer/Private/Views/List/SObjectMixerEditorList.cpp
+++ b/Engine/Plugins/Editor/ObjectMixer/ObjectMixer/Source/ObjectMixer/Private/Views/List/SObjectMixerEditorList.cpp
@@ -439,9 +439,9 @@ TSet<TSharedPtr<ISceneOutlinerTreeItem>> SObjectMixerEditorList::GetSoloRows() c
 
 TSet<TSharedPtr<ISceneOutlinerTreeItem>> SObjectMixerEditorList::GetTreeRootItems() const
 {
-	if (GetTree().GetItems().Num())
+	if (GetTree().GetRootItems().Num())
 	{
-		return {GetTree().GetItems()[0]};
+		return {GetTree().GetRootItems()[0]};
 	}
 
 	return {};
@@ -449,9 +449,9 @@ TSet<TSharedPtr<ISceneOutlinerTreeItem>> SObjectMixerEditorList::GetTreeRootItem
 
 TSet<TWeakPtr<ISceneOutlinerTreeItem>> SObjectMixerEditorList::GetWeakTreeRootItems() const
 {
-	if (GetTree().GetItems().Num())
+	if (GetTree().GetRootItems().Num())
 	{
-		return {GetTree().GetItems()[0]};
+		return {GetTree().GetRootItems()[0]};
 	}
 
 	return {};
diff --git a/Engine/Plugins/Editor/ObjectMixer/ObjectMixer/Source/ObjectMixer/Public/Views/List/SObjectMixerEditorList.h b/Engine/Plugins/Editor/ObjectMixer/ObjectMixer/Source/ObjectMixer/Public/Views/List/SObjectMixerEditorList.h
--- a/Engine/Plugins/Editor/ObjectMixer/ObjectMixer/Source/ObjectMixer/Public/Views/List/SObjectMixerEditorList.h
+++ b/Engine/Plugins/Editor/ObjectMixer/ObjectMixer/Source/ObjectMixer/Public/Views/List/SObjectMixerEditorList.h
@@ -137,11 +137,6 @@ public:
 
 	bool IsTreeViewItemSelected(TSharedPtr<ISceneOutlinerTreeItem> Item);
 
-	UE_NODISCARD int32 GetTreeViewItemCount() const
-	{
-		return GetTree().GetItems().Num();
-	}
-
 	TSet<TSharedPtr<ISceneOutlinerTreeItem>> GetTreeRootItems() const;
 	TSet<TWeakPtr<ISceneOutlinerTreeItem>> GetWeakTreeRootItems() const;
 	
diff --git a/Engine/Plugins/Experimental/CommonUI/Source/CommonUI/Public/CommonTreeView.h b/Engine/Plugins/Experimental/CommonUI/Source/CommonUI/Public/CommonTreeView.h
--- a/Engine/Plugins/Experimental/CommonUI/Source/CommonUI/Public/CommonTreeView.h
+++ b/Engine/Plugins/Experimental/CommonUI/Source/CommonUI/Public/CommonTreeView.h
@@ -21,12 +21,12 @@ public:
 		{
 			// Set selection to the first item in a list if no items are selected.
 			// If bReturnFocusToSelection is true find the last selected object and focus on that.
-			if (this->HasValidItemsSource() && this->TreeItemsSource->Num() > 0)
+			if (this->HasValidItemsSource() && this->GetRootItems().Num() > 0)
 			{
 				typename TListTypeTraits<ItemType>::NullableType ItemNavigatedTo = TListTypeTraits<ItemType>::MakeNullPtr();
 				if (this->GetNumItemsSelected() == 0)
 				{
-					ItemNavigatedTo = (*this->TreeItemsSource)[0];
+					ItemNavigatedTo = this->GetRootItems()[0];
 				}
 				else if (this->bReturnFocusToSelection && TListTypeTraits<ItemType>::IsPtrValid(this->SelectorItem))
 				{
diff --git a/Engine/Plugins/Experimental/RenderGrid/Source/RenderGridEditor/Private/UI/Components/SRenderGridDragHandle.h b/Engine/Plugins/Experimental/RenderGrid/Source/RenderGridEditor/Private/UI/Components/SRenderGridDragHandle.h
--- a/Engine/Plugins/Experimental/RenderGrid/Source/RenderGridEditor/Private/UI/Components/SRenderGridDragHandle.h
+++ b/Engine/Plugins/Experimental/RenderGrid/Source/RenderGridEditor/Private/UI/Components/SRenderGridDragHandle.h
@@ -17,11 +17,11 @@ template<typename DragDropType>
 class SRenderGridDragHandle : public SCompoundWidget
 {
 public:
-	using WidgetType = typename DragDropType::WidgetType;
+	using DragDropWidgetType = typename DragDropType::WidgetType;
 	using HeldItemType = typename DragDropType::HeldItemType;
 
 	SLATE_BEGIN_ARGS(SRenderGridDragHandle) {}
-		SLATE_ARGUMENT(TSharedPtr<WidgetType>, Widget)
+		SLATE_ARGUMENT(TSharedPtr<DragDropWidgetType>, Widget)
 	SLATE_END_ARGS()
 
 	void Construct(const FArguments& InArgs, HeldItemType InItem)
@@ -64,7 +64,7 @@ public:
 
 	TSharedPtr<FDragDropOperation> CreateDragDropOperation()
 	{
-		if (const TSharedPtr<WidgetType> Widget = WidgetWeakPtr.Pin())
+		if (const TSharedPtr<DragDropWidgetType> Widget = WidgetWeakPtr.Pin())
 		{
 			TSharedPtr<DragDropType> DragDropOperation = MakeShared<DragDropType>(Widget, Item);
 			DragDropOperation->Construct();
@@ -75,7 +75,7 @@ public:
 
 private:
 	/** Holds the widget to display when dragging. */
-	TWeakPtr<WidgetType> WidgetWeakPtr;
+	TWeakPtr<DragDropWidgetType> WidgetWeakPtr;
 
 	/** Holds the item being dragged. */
 	HeldItemType Item;
diff --git a/Engine/Plugins/VirtualProduction/RemoteControl/Source/RemoteControlUI/Private/UI/SRCPanelDragHandle.h b/Engine/Plugins/VirtualProduction/RemoteControl/Source/RemoteControlUI/Private/UI/SRCPanelDragHandle.h
--- a/Engine/Plugins/VirtualProduction/RemoteControl/Source/RemoteControlUI/Private/UI/SRCPanelDragHandle.h
+++ b/Engine/Plugins/VirtualProduction/RemoteControl/Source/RemoteControlUI/Private/UI/SRCPanelDragHandle.h
@@ -17,10 +17,10 @@ template <typename DragDropType>
 class SRCPanelDragHandle : public SCompoundWidget
 {
 public:
-	using WidgetType = typename DragDropType::WidgetType;
+	using DragDropWidgetType = typename DragDropType::WidgetType;
 	SLATE_BEGIN_ARGS(SRCPanelDragHandle)
 	{}
-		SLATE_ARGUMENT(TSharedPtr<WidgetType>, Widget)
+		SLATE_ARGUMENT(TSharedPtr<DragDropWidgetType>, Widget)
 	SLATE_END_ARGS()
 
 	void Construct(const FArguments& InArgs, FGuid InId)
@@ -66,7 +66,7 @@ public:
 
 	TSharedPtr<FDragDropOperation> CreateDragDropOperation()
 	{
-		if (TSharedPtr<WidgetType> WidgetPtr = Widget.Pin())
+		if (TSharedPtr<DragDropWidgetType> WidgetPtr = Widget.Pin())
 		{
 			TSharedPtr<FDecoratedDragDropOp> DragDropOperation = MakeShared<DragDropType>(Widget.Pin(), Id);
 			DragDropOperation->Construct();
@@ -77,7 +77,7 @@ public:
 
 private:
 	/** Holds the widget to display when dragging. */
-	TWeakPtr<WidgetType> Widget;
+	TWeakPtr<DragDropWidgetType> Widget;
 	/** Holds the ID of the item being dragged. */
 	FGuid Id;
 };
diff --git a/Engine/Source/Runtime/AppFramework/Private/Framework/Testing/SStarshipGallery.cpp b/Engine/Source/Runtime/AppFramework/Private/Framework/Testing/SStarshipGallery.cpp
--- a/Engine/Source/Runtime/AppFramework/Private/Framework/Testing/SStarshipGallery.cpp
+++ b/Engine/Source/Runtime/AppFramework/Private/Framework/Testing/SStarshipGallery.cpp
@@ -18,7 +18,7 @@
 
 // #include "SlateCore.h"
 #include "Brushes/SlateRoundedBoxBrush.h"
-
+#include "Containers/ObservableArray.h"
 
 #include "Layout/WidgetPath.h"
 #include "SlateOptMacros.h"
@@ -58,6 +58,9 @@
 #include "Widgets/Input/SComboBox.h"
 #include "Widgets/Input/SNumericEntryBox.h"
 #include "Widgets/Input/SSegmentedControl.h"
+#include "Widgets/Views/SListView.h"
+#include "Widgets/Views/STileView.h"
+#include "Widgets/Views/STreeView.h"
 
 
 #include "Brushes/SlateImageBrush.h"
@@ -112,6 +115,7 @@ public:
                 ->AddTab("Text", ETabState::OpenedTab)
                 ->AddTab("Icons", ETabState::OpenedTab)
                 ->AddTab("Starship Widgets", ETabState::OpenedTab)
+                ->AddTab("List Widgets", ETabState::OpenedTab)
                 // ->AddTab("SLATE WIDGETS", ETabState::OpenedTab)
                 ->SetForegroundTab(FName("Starship Widgets"))
             )
@@ -176,6 +180,21 @@ public:
             )
         );
 
+		FGlobalTabmanager::Get()->RegisterNomadTabSpawner("List Widgets",
+			FOnSpawnTab::CreateLambda(
+				[this](const FSpawnTabArgs&) -> TSharedRef<SDockTab>
+				{
+					return SNew(SDockTab)
+						.TabRole(ETabRole::PanelTab)
+						.ContentPadding(0)
+						.ForegroundColor(FSlateColor::UseStyle())
+						[
+							ConstructListGallery()
+						];
+				}
+			)
+		);
+
 
         /*
         FGlobalTabmanager::Get()->RegisterTabSpawner("SLATE WIDGETS", 
@@ -720,52 +739,51 @@ public:
 
         TSharedPtr<SGridPanel> WidgetGrid = SNew(SGridPanel);
 
-        int WidgetNum = 0, RowCount = 15, Cols = 3;
-        auto NextSlot = [WidgetNum, RowCount, Cols](TSharedPtr<SGridPanel> Grid, const FText& InLabel) mutable -> SHorizontalBox::FScopedWidgetSlotArguments
-        { 
-            TSharedRef<SHorizontalBox> HBox = SNew(SHorizontalBox);
-
-            // Checkbox to show disabled state
-            Grid->AddSlot((WidgetNum / RowCount)*Cols, WidgetNum % RowCount)
-            .Padding(12.f)
-            .HAlign(HAlign_Left)
-            .VAlign(VAlign_Center)
-            [
-		SNew(SCheckBox)
-                .IsChecked_Lambda( [HBox] { return HBox->IsEnabled() ? ECheckBoxState::Checked : ECheckBoxState::Unchecked; } )
-                .OnCheckStateChanged_Lambda( [HBox] (ECheckBoxState NewState) { HBox->SetEnabled( NewState == ECheckBoxState::Checked); })
-            ];
-
-            // Add the Label
-            Grid->AddSlot((WidgetNum / RowCount)*Cols + 1, WidgetNum % RowCount )
-            .Padding(24.f, 16.f, 12.f, 16.f)
-            .VAlign(VAlign_Center)
-            .HAlign(HAlign_Right)
-            [
-                SNew(STextBlock)
-                .Font(FAppStyle::Get().GetFontStyle("NormalFont"))
-                .ColorAndOpacity(FAppStyle::Get().GetSlateColor("Colors.White50"))
-                .Text(InLabel)
-            ];
-
-            //auto& Ret = Grid->AddSlot((WidgetNum / RowCount)*2 + 1, WidgetNum % RowCount)
-            Grid->AddSlot((WidgetNum / RowCount)*Cols + 2, WidgetNum % RowCount)
-            .Padding(12.f, 16.f, 12.f, 16.f)
-            .VAlign(VAlign_Center)
-            .HAlign(HAlign_Fill)
-            [
+		int32 WidgetNum = 0, RowCount = 15, Cols = 3;
+		auto NextSlot = [WidgetNum, RowCount, Cols](TSharedPtr<SGridPanel> Grid, const FText& InLabel) mutable -> SHorizontalBox::FScopedWidgetSlotArguments
+		{
+			TSharedRef<SHorizontalBox> HBox = SNew(SHorizontalBox);
 
-                HBox
-            ];
+			// Checkbox to show disabled state
+			Grid->AddSlot((WidgetNum / RowCount) * Cols, WidgetNum % RowCount)
+				.Padding(12.f)
+				.HAlign(HAlign_Left)
+				.VAlign(VAlign_Center)
+				[
+					SNew(SCheckBox)
+					.IsChecked_Lambda([HBox] { return HBox->IsEnabled() ? ECheckBoxState::Checked : ECheckBoxState::Unchecked; })
+					.OnCheckStateChanged_Lambda([HBox](ECheckBoxState NewState) { HBox->SetEnabled(NewState == ECheckBoxState::Checked); })
+				];
+
+			// Add the Label
+			Grid->AddSlot((WidgetNum / RowCount) * Cols + 1, WidgetNum % RowCount)
+				.Padding(24.f, 16.f, 12.f, 16.f)
+				.VAlign(VAlign_Center)
+				.HAlign(HAlign_Right)
+				[
+					SNew(STextBlock)
+					.Font(FAppStyle::Get().GetFontStyle("NormalFont"))
+					.ColorAndOpacity(FAppStyle::Get().GetSlateColor("Colors.White50"))
+					.Text(InLabel)
+				];
+
+			//auto& Ret = Grid->AddSlot((WidgetNum / RowCount)*2 + 1, WidgetNum % RowCount)
+			Grid->AddSlot((WidgetNum / RowCount) * Cols + 2, WidgetNum % RowCount)
+				.Padding(12.f, 16.f, 12.f, 16.f)
+				.VAlign(VAlign_Center)
+				.HAlign(HAlign_Fill)
+				[
+					HBox
+				];
 
-            ++WidgetNum;
+			++WidgetNum;
 
 			SHorizontalBox::FScopedWidgetSlotArguments NewSlot = HBox->AddSlot();
 			NewSlot.VAlign(VAlign_Center)
 				.HAlign(HAlign_Fill)
 				.AutoWidth();
 			return MoveTemp(NewSlot);
-        };
+		};
 
 
         auto LeftRightLabel = [](const FName& InIconName = FName(), const FText& InLabel = FText::GetEmpty(), const FName& InTextStyle = TEXT("ButtonText")) -> TSharedRef<SWidget>
@@ -885,14 +903,12 @@ public:
                     LeftRightLabel("Icons.box-perspective")
                 ]
             ]
-
         ];
 
 
         // NoBorder Button Button 
         NextSlot(WidgetGrid, LOCTEXT("TextButton", "Simple Button"))
         [
-
             SNew(SHorizontalBox)
             +SHorizontalBox::Slot()
             .AutoWidth()
@@ -1478,9 +1494,157 @@ public:
             ]
         ];
 
-    };
+    }
+
+	class SListGalleryWidget : public SCompoundWidget
+	{
+	private:
+		TArray<TSharedPtr<int32>> ListViewItemsInArray;
+		UE::Slate::Containers::TObservableArray<TSharedPtr<int32>> ListViewItemsInObservableArray;
+		TSharedPtr<UE::Slate::Containers::TObservableArray<TSharedPtr<int32>>> SharedListViewItemsInObservableArray;
+		TSharedPtr<SListView<TSharedPtr<int32>>> ListView_List;
+		TSharedPtr<STileView<TSharedPtr<int32>>> ListView_Tile;
+		TSharedPtr<STreeView<TSharedPtr<int32>>> ListView_Tree;
+
+	public:
+		SLATE_BEGIN_ARGS(SListGalleryWidget){}
+		SLATE_END_ARGS()
+		void Construct(const FArguments&)
+		{
+			TSharedRef<SHorizontalBox> TheBox = SNew(SHorizontalBox);
+
+			SharedListViewItemsInObservableArray = MakeShared<UE::Slate::Containers::TObservableArray<TSharedPtr<int32>>>();
+			for (int32 Index = 0; Index < 15; ++Index)
+			{
+				ListViewItemsInArray.Add(MakeShared<int32>(Index));
+				ListViewItemsInObservableArray.Add(MakeShared<int32>(Index));
+				SharedListViewItemsInObservableArray->Add(MakeShared<int32>(Index));
+			}
+
+			TheBox->AddSlot()
+			[
+				SNew(SVerticalBox)
+				+ SVerticalBox::Slot()
+				[
+					SNew(SButton)
+					.OnClicked_Lambda([this]()
+					{
+						ListViewItemsInArray.Add(MakeShared<int32>(ListViewItemsInArray.Num()));
+						ListViewItemsInObservableArray.Add(MakeShared<int32>(ListViewItemsInObservableArray.Num()));
+						SharedListViewItemsInObservableArray->Add(MakeShared<int32>(SharedListViewItemsInObservableArray->Num()));
+						return FReply::Handled();
+					})
+					.Content()
+					[
+						SNew(STextBlock)
+						.Text(LOCTEXT("Add", "Add"))
+					]
+				]
+				+ SVerticalBox::Slot()
+				[
+					SNew(SButton)
+					.OnClicked_Lambda([this]()
+					{
+						int32 Rand = FMath::RandHelper(ListViewItemsInArray.Num());
+						ListViewItemsInArray.RemoveAt(Rand);
+						ListViewItemsInObservableArray.RemoveAt(Rand);
+						SharedListViewItemsInObservableArray->RemoveAt(Rand);
+						return FReply::Handled();
+					})
+					.Content()
+					[
+						SNew(STextBlock)
+						.Text(LOCTEXT("Remove", "Remove"))
+					]
+				]
+			];
+
+			TheBox->AddSlot()
+			[
+				SNew(SVerticalBox)
+				+ SVerticalBox::Slot()
+				[
+					SAssignNew(ListView_List, SListView<TSharedPtr<int32>>)
+					.ListItemsSource(&ListViewItemsInArray)
+					.ListItemsSource(&ListViewItemsInObservableArray)
+					.ListItemsSource(SharedListViewItemsInObservableArray)
+					.OnGenerateRow_Lambda([](TSharedPtr<int32> Value, const TSharedRef<STableViewBase>& OwnerTable)
+					{
+						typedef STableRow<TSharedPtr<int32>> RowType;
+
+						TSharedRef<RowType> NewRow = SNew(RowType, OwnerTable);
+						NewRow->SetContent(SNew(STextBlock).Text(FText::AsNumber(*Value.Get())));
+
+						return NewRow;
+					})
+				]
+				+ SVerticalBox::Slot()
+				[
+					SAssignNew(ListView_Tile, STileView<TSharedPtr<int32>>)
+					.ListItemsSource(&ListViewItemsInArray)
+					.ListItemsSource(&ListViewItemsInObservableArray)
+					.ListItemsSource(SharedListViewItemsInObservableArray)
+					.OnGenerateTile_Lambda([](TSharedPtr<int32> Value, const TSharedRef<STableViewBase>& OwnerTable)
+					{
+						typedef STableRow<TSharedPtr<int32>> RowType;
+
+						TSharedRef<RowType> NewRow = SNew(RowType, OwnerTable);
+						NewRow->SetContent(SNew(STextBlock).Text(FText::AsNumber(*Value.Get())));
+
+						return NewRow;
+					})
+				]
+				+ SVerticalBox::Slot()
+				[
+					SAssignNew(ListView_Tree, STreeView<TSharedPtr<int32>>)
+					.TreeItemsSource(&ListViewItemsInArray)
+					.TreeItemsSource(&ListViewItemsInObservableArray)
+					.TreeItemsSource(SharedListViewItemsInObservableArray)
+					.OnGenerateRow_Lambda([](TSharedPtr<int32> Value, const TSharedRef<STableViewBase>& OwnerTable)
+					{
+						typedef STableRow<TSharedPtr<int32>> RowType;
+
+						TSharedRef<RowType> NewRow = SNew(RowType, OwnerTable);
+						NewRow->SetContent(SNew(STextBlock).Text(FText::AsNumber(*Value.Get())));
+
+						return NewRow;
+					})
+					.OnGetChildren_Lambda([](TSharedPtr<int32> Parent, TArray<TSharedPtr<int32>>& OutParent)
+					{
+						OutParent.Add(MakeShared<int32>(99));
+					})
+				]
+			];
+
+			ChildSlot[TheBox];
+		}
+		~SListGalleryWidget()
+		{
+			if (ListView_List)
+			{
+				ListView_List->ClearItemsSource();
+				ListView_List.Reset();
+			}
+			if (ListView_Tile)
+			{
+				ListView_Tile->ClearItemsSource();
+				ListView_Tile.Reset();
+			}
+			if (ListView_Tree)
+			{
+				ListView_Tree->ClearRootItemsSource();
+				ListView_Tree.Reset();
+			}
+			ChildSlot[SNullWidget::NullWidget];
+		}
+	};
+	
+	TSharedRef<SWidget> ConstructListGallery()
+	{
+		return SNew(SListGalleryWidget);
+	}
 
-    TArray< TSharedPtr<EHorizontalAlignment> > HorizontalAlignmentComboItems;
+	TArray<TSharedPtr<EHorizontalAlignment>> HorizontalAlignmentComboItems;
 
 }; // class SStarshipGallery
 
diff --git a/Engine/Source/Runtime/Slate/Public/Widgets/Views/SListView.h b/Engine/Source/Runtime/Slate/Public/Widgets/Views/SListView.h
--- a/Engine/Source/Runtime/Slate/Public/Widgets/Views/SListView.h
+++ b/Engine/Source/Runtime/Slate/Public/Widgets/Views/SListView.h
@@ -4,25 +4,26 @@
 
 #include "CoreMinimal.h"
 #include "Containers/ArrayView.h"
+#include "Containers/ObservableArray.h"
 #include "InputCoreTypes.h"
 #include "Input/Reply.h"
 #include "Layout/Visibility.h"
-#include "Widgets/DeclarativeSyntaxSupport.h"
-#include "Widgets/SWidget.h"
 #include "Styling/SlateTypes.h"
 #include "Styling/AppStyle.h"
 #include "Framework/SlateDelegates.h"
-#include "Widgets/Text/STextBlock.h"
+#include "Framework/Layout/Overscroll.h"
 #include "Framework/Views/ITypedTableView.h"
 #include "Framework/Views/TableViewMetadata.h"
-#include "Widgets/Views/STableViewBase.h"
 #include "Framework/Views/TableViewTypeTraits.h"
-#include "Widgets/Views/STableRow.h"
 #include "Types/SlateConstants.h"
+#include "Widgets/DeclarativeSyntaxSupport.h"
+#include "Widgets/SOverlay.h"
 #include "Widgets/Layout/SScrollBar.h"
-#include "Framework/Layout/Overscroll.h"
 #include "Widgets/Images/SImage.h"
-#include "Widgets/SOverlay.h"
+#include "Widgets/Text/STextBlock.h"
+#include "Widgets/Views/STableViewBase.h"
+#include "Widgets/Views/STableRow.h"
+#include "Widgets/Views/IItemsSource.h"
 #include "Application/SlateApplicationBase.h"
 #if WITH_ACCESSIBILITY
 #include "GenericPlatform/Accessibility/GenericAccessibleInterfaces.h"
@@ -90,7 +91,6 @@ public:
 		, _OnGeneratePinnedRow()
 		, _OnEntryInitialized()
 		, _OnRowReleased()
-		, _ListItemsSource()
 		, _ItemHeight(16)
 		, _MaxPinnedItems(6)
 		, _OnContextMenuOpening()
@@ -137,7 +137,7 @@ public:
 		
 		SLATE_EVENT( FOnFinishedScrolling, OnFinishedScrolling )
 
-		SLATE_ARGUMENT( const TArray<ItemType>* , ListItemsSource )
+		SLATE_ITEMS_SOURCE_ARGUMENT( ItemType, ListItemsSource )
 
 		SLATE_ATTRIBUTE( float, ItemHeight )
 
@@ -219,9 +219,7 @@ public:
 		this->OnItemScrolledIntoView = InArgs._OnItemScrolledIntoView;
 		this->OnFinishedScrolling = InArgs._OnFinishedScrolling;
 
-PRAGMA_DISABLE_DEPRECATION_WARNINGS
-		this->SetItemsSource(InArgs._ListItemsSource);
-PRAGMA_ENABLE_DEPRECATION_WARNINGS
+		this->SetItemsSource(InArgs.MakeListItemsSource(this->SharedThis(this)));
 		this->OnContextMenuOpening = InArgs._OnContextMenuOpening;
 		this->OnClick = InArgs._OnMouseButtonClick;
 		this->OnDoubleClick = InArgs._OnMouseButtonDoubleClick;
@@ -318,7 +316,7 @@ PRAGMA_ENABLE_DEPRECATION_WARNINGS
 
 public:
 
-	// SWidget overrides
+	//~ SWidget overrides
 
 	virtual bool SupportsKeyboardFocus() const override
 	{
@@ -586,7 +584,7 @@ public:
 	}
 
 #if WITH_ACCESSIBILITY
-	protected:
+protected:
 	friend class FSlateAccessibleListView;
 	/**
 	* An accessible implementation for SListView to be exposed to platform accessibility APIs.
@@ -656,7 +654,7 @@ public:
 		}
 		// ~
 	};
-	public:
+public:
 	virtual TSharedRef<FSlateAccessibleWidget> CreateAccessibleWidget() override
 	{
 		// @TODOAccessibility: Add support for the different types of tables e.g tree and tile 
@@ -696,7 +694,7 @@ private:
 		 * @param Item  The item for which to find the widget.
 		 * @return A pointer to the corresponding widget if it exists; otherwise nullptr.
 		 */
-		TSharedPtr<ITableRow> GetWidgetForItem( const ItemType& Item ) const
+		[[nodiscard]] TSharedPtr<ITableRow> GetWidgetForItem( const ItemType& Item ) const
 		{
 			const TSharedRef<ITableRow>* LookupResult = ItemToWidgetMap.Find(Item);
 			return LookupResult ? TSharedPtr<ITableRow>(*LookupResult) : TSharedPtr<ITableRow>(nullptr);
@@ -737,7 +735,7 @@ private:
 		 */
 		void OnBeginGenerationPass()
 		{
-			// Assume all the previously generated items need to be cleaned up.				
+			// Assume all the previously generated items need to be cleaned up.
 			ItemsToBeCleanedUp = ItemsWithGeneratedWidgets;
 			ItemsWithGeneratedWidgets.Empty();
 		}
@@ -862,6 +860,7 @@ private:
 			}			
 		}
 
+	public:
 		/** We store a pointer to the owner list for error purposes, so when asserts occur we can report which list it happened for. */
 		SListView<ItemType>* OwnerList;
 
@@ -1666,21 +1665,62 @@ public:
 	 */
 	void SetItemsSource(const TArray<ItemType>* InListItemsSource)
 	{
-		PRAGMA_DISABLE_DEPRECATION_WARNINGS
-		if (ItemsSource != InListItemsSource)
+		ensureMsgf(InListItemsSource, TEXT("The ListItems is invalid."));
+		if (ViewSource == nullptr || !ViewSource->IsSame(reinterpret_cast<const void*>(InListItemsSource)))
 		{
-			if (IsConstructed())
+			if (InListItemsSource)
 			{
-				Private_ClearSelection();
-				CancelScrollIntoView();
-				ClearWidgets();
-				RebuildList();
+				SetItemsSource(MakeUnique<UE::Slate::ItemsSource::FArrayPointer<ItemType>>(InListItemsSource));
 			}
-			ItemsSource = InListItemsSource;
+			else
+			{
+				ClearItemsSource();
+			}
+		}
+	}
+
+	/**
+	 * Establishes a wholly new list of items being observed by the list.
+	 * Wipes all existing state and requests and will fully rebuild on the next tick.
+	 * The ObservableArray will notify the Widget when it needs to refresh.
+	 */
+	void SetItemsSource(TSharedRef<::UE::Slate::Containers::TObservableArray<ItemType>> InListItemsSource)
+	{
+		if (ViewSource == nullptr || !ViewSource->IsSame(reinterpret_cast<const void*>(&InListItemsSource.Get())))
+		{
+			SetItemsSource(MakeUnique<UE::Slate::ItemsSource::FSharedObservableArray<ItemType>>(SharedThis(this), MoveTemp(InListItemsSource)));
 		}
-		PRAGMA_ENABLE_DEPRECATION_WARNINGS
 	}
 
+	/**
+	 * Establishes a wholly new list of items being observed by the list.
+	 * Wipes all existing state and requests and will fully rebuild on the next tick.
+	 */
+	void SetItemsSource(TUniquePtr<UE::Slate::ItemsSource::IItemsSource<ItemType>> Provider)
+	{
+		if (IsConstructed())
+		{
+			Private_ClearSelection();
+			CancelScrollIntoView();
+			ClearWidgets();
+
+			ViewSource = MoveTemp(Provider);
+
+			RebuildList();
+		}
+		else
+		{
+			ViewSource = MoveTemp(Provider);
+		}
+	}
+
+	void ClearItemsSource()
+	{
+		SetItemsSource(TUniquePtr<UE::Slate::ItemsSource::IItemsSource<ItemType>>());
+	}
+
+public:
+
 	UE_DEPRECATED(5.2, "SetListItemsSource is deprecated. Please use the correct SetItemsSource implementation.")
 	void SetListItemsSource(const TArray<ItemType>& InListItemsSource)
 	{
@@ -1689,20 +1729,12 @@ public:
 
 	bool HasValidItemsSource() const
 	{
-		PRAGMA_DISABLE_DEPRECATION_WARNINGS
-		return ItemsSource != nullptr;
-		PRAGMA_ENABLE_DEPRECATION_WARNINGS
+		return ViewSource != nullptr;
 	}
 
 	TArrayView<const ItemType> GetItems() const
 	{
-		PRAGMA_DISABLE_DEPRECATION_WARNINGS
-		if (ItemsSource)
-		{
-			return *ItemsSource;
-		}
-		return TArrayView<const ItemType>();
-		PRAGMA_ENABLE_DEPRECATION_WARNINGS
+		return ViewSource ? ViewSource->GetItems() : TArrayView<const ItemType>();
 	}
 
 	/**
@@ -2585,6 +2617,10 @@ protected:
 	/** If true, number of pinned items > MaxPinnedItems so some items are collapsed in the hierarchy */
 	bool bIsHierarchyCollapsed = false;
 
+private:
+	/** Pointer to the source data that we are observing */
+	TUniquePtr<UE::Slate::ItemsSource::IItemsSource<ItemType>> ViewSource;
+
 private:
 	struct FGenerationPassGuard
 	{
diff --git a/Engine/Source/Runtime/Slate/Public/Widgets/Views/STileView.h b/Engine/Source/Runtime/Slate/Public/Widgets/Views/STileView.h
--- a/Engine/Source/Runtime/Slate/Public/Widgets/Views/STileView.h
+++ b/Engine/Source/Runtime/Slate/Public/Widgets/Views/STileView.h
@@ -41,7 +41,6 @@ public:
 	SLATE_BEGIN_ARGS(STileView<ItemType>)
 		: _OnGenerateTile()
 		, _OnTileReleased()
-		, _ListItemsSource(static_cast<TArray<ItemType>*>(nullptr)) //@todo Slate Syntax: Initializing from nullptr without a cast
 		, _ItemHeight(128)
 		, _ItemWidth(128)
 		, _ItemAlignment(EListItemAlignment::EvenlyDistributed)
@@ -79,7 +78,7 @@ public:
 
 		SLATE_EVENT( FOnItemScrolledIntoView, OnItemScrolledIntoView )
 
-		SLATE_ARGUMENT( const TArray<ItemType>* , ListItemsSource )
+		SLATE_ITEMS_SOURCE_ARGUMENT( ItemType, ListItemsSource )
 
 		SLATE_ATTRIBUTE( float, ItemHeight )
 
@@ -149,7 +148,7 @@ public:
 		this->OnRowReleased = InArgs._OnTileReleased;
 		this->OnItemScrolledIntoView = InArgs._OnItemScrolledIntoView;
 		
-		this->SetItemsSource(InArgs._ListItemsSource);
+		this->SetItemsSource(InArgs.MakeListItemsSource(this->SharedThis(this)));
 		this->OnContextMenuOpening = InArgs._OnContextMenuOpening;
 		this->OnClick = InArgs._OnMouseButtonClick;
 		this->OnDoubleClick = InArgs._OnMouseButtonDoubleClick;
@@ -270,14 +269,14 @@ public:
 		// Clear all the items from our panel. We will re-add them in the correct order momentarily.
 		this->ClearWidgets();
 		
-		const TArrayView<const ItemType> tems = this->GetItems();
-		if (tems.Num() > 0)
+		const TArrayView<const ItemType> Items = this->GetItems();
+		if (Items.Num() > 0)
 		{
 			// Item width and height is constant by design.
 			FTableViewDimensions TileDimensions = GetTileDimensions();
 			FTableViewDimensions AllottedDimensions(this->Orientation, MyGeometry.GetLocalSize());
 
-			const int32 NumItems = tems.Num();
+			const int32 NumItems = Items.Num();
 			const int32 NumItemsPerLine = GetNumItemsPerLine();
 			const int32 NumItemsPaddedToFillLastLine = (NumItems % NumItemsPerLine != 0)
 				? NumItems + NumItemsPerLine - NumItems % NumItemsPerLine
@@ -305,7 +304,7 @@ public:
 			double NumLinesShownOnScreen = 0;
 			for( int32 ItemIndex = StartIndex; !bHasFilledAvailableArea && ItemIndex < NumItems; ++ItemIndex )
 			{
-				const ItemType& CurItem = tems[ItemIndex];
+				const ItemType& CurItem = Items[ItemIndex];
 
 				if (bNewLine)
 				{
diff --git a/Engine/Source/Runtime/Slate/Public/Widgets/Views/STreeView.h b/Engine/Source/Runtime/Slate/Public/Widgets/Views/STreeView.h
--- a/Engine/Source/Runtime/Slate/Public/Widgets/Views/STreeView.h
+++ b/Engine/Source/Runtime/Slate/Public/Widgets/Views/STreeView.h
@@ -119,7 +119,6 @@ public:
 		, _OnGeneratePinnedRow()
 		, _OnGetChildren()
 		, _OnSetExpansionRecursive()
-		, _TreeItemsSource( static_cast< const TArray<ItemType>* >(nullptr) ) //@todo Slate Syntax: Initializing from nullptr without a cast
 		, _ItemHeight(16)
 		, _MaxPinnedItems(6) // Having more than the max amount of items leads to the extra items in the middle being collapsed into ellipses, and the last item is fully shown
 		, _OnContextMenuOpening()
@@ -166,7 +165,7 @@ public:
 
 		SLATE_EVENT( FOnSetExpansionRecursive, OnSetExpansionRecursive )
 
-		SLATE_ARGUMENT( const TArray<ItemType>*, TreeItemsSource )
+		SLATE_ITEMS_SOURCE_ARGUMENT( ItemType, TreeItemsSource )
 
 		SLATE_ATTRIBUTE( float, ItemHeight )
 
@@ -247,7 +246,8 @@ public:
 		this->OnItemScrolledIntoView = InArgs._OnItemScrolledIntoView;
 		this->OnGetChildren = InArgs._OnGetChildren;
 		this->OnSetExpansionRecursive = InArgs._OnSetExpansionRecursive;
-		this->TreeItemsSource = InArgs._TreeItemsSource;
+
+		this->SetRootItemsSource(InArgs.MakeTreeItemsSource(this->SharedThis(this)));
 
 		this->OnKeyDownHandler = InArgs._OnKeyDownHandler;
 		this->OnContextMenuOpening = InArgs._OnContextMenuOpening;
@@ -294,7 +294,7 @@ public:
 				ErrorString += TEXT("Please specify an OnGenerateRow. \n");
 			}
 
-			if ( this->TreeItemsSource == nullptr )
+			if (!this->HasValidRootItemsSource())
 			{
 				ErrorString += TEXT("Please specify a TreeItemsSource. \n");
 			}
@@ -331,15 +331,14 @@ public:
 	/** Default constructor. */
 	STreeView()
 		: SListView< ItemType >( ETableViewMode::Tree )
-		, TreeItemsSource( nullptr )
 		, bTreeItemsAreDirty( true )
 	{
-		this->SetItemsSource(&LinearizedItems);
+		SListView<ItemType>::SetItemsSource(&LinearizedItems);
 	}
 
 public:
 
-	// SWidget overrides
+	//~ SWidget overrides
 
 	virtual FReply OnKeyDown( const FGeometry& MyGeometry, const FKeyEvent& InKeyEvent ) override
 	{
@@ -423,8 +422,8 @@ public:
 	
 private:
 
-	// Tree View adds the ability to expand/collapse items.
-	// All the selection functionality is inherited from ListView.
+	//~ Tree View adds the ability to expand/collapse items.
+	//~ All the selection functionality is inherited from ListView.
 
 	virtual bool Private_IsItemExpanded( const ItemType& TheItem ) const override
 	{
@@ -565,7 +564,7 @@ public:
 				// We are about to repopulate linearized items; the ListView that TreeView is built on top of will also need to refresh.
 				bTreeItemsAreDirty = false;
 
-				if ( OnGetChildren.IsBound() && TreeItemsSource != nullptr )								
+				if ( OnGetChildren.IsBound() && HasValidRootItemsSource() )
 				{
 					// We make copies of the old expansion and selection sets so that we can remove
 					// any items that are no longer seen by the tree.
@@ -575,7 +574,7 @@ public:
 					
 					// Rebuild the linearized view of the tree data.
 					LinearizedItems.Empty();
-					PopulateLinearizedItems( *TreeItemsSource, LinearizedItems, TempDenseItemInfos, TBitArray<>(), TempSelectedItemsMap, TempSparseItemInfo, true, INDEX_NONE );
+					PopulateLinearizedItems(GetRootItems(), LinearizedItems, TempDenseItemInfos, TBitArray<>(), TempSelectedItemsMap, TempSparseItemInfo, true, INDEX_NONE);
 
 					if( !bAllowInvisibleItemSelection &&
 						(this->SelectedItems.Num() != TempSelectedItemsMap.Num() ||
@@ -638,9 +637,9 @@ public:
 	 * @return true if we encountered expanded children; false otherwise.
 	 */
 	bool PopulateLinearizedItems(
-		const TArray<ItemType>& InItemsSource,
-		TArray< ItemType >& InLinearizedItems,
-		TArray< FItemInfo >& NewDenseItemInfos,
+		TArrayView<const ItemType> InItemsSource,
+		TArray<ItemType>& InLinearizedItems,
+		TArray<FItemInfo>& NewDenseItemInfos,
 		TBitArray<> NeedsParentWire,
 		TItemSet& OutNewSelectedItems,
 		TSparseItemMap& NewSparseItemInfo,
@@ -899,18 +898,81 @@ public:
 		return Private_IsItemExpanded( InItem );
 	}
 
-		
+protected:
+	//~ Hide the base SetItemsSource
+	using SListView<ItemType>::SetItemsSource;
+	using SListView<ItemType>::ClearItemsSource;
+	using SListView<ItemType>::HasValidItemsSource;
+	using SListView<ItemType>::GetItems;
+
+public:
+
 	/**
 	 * Set the TreeItemsSource. The Tree will generate widgets to represent these items.
-	 *
 	 * @param InItemsSource  A pointer to the array of items that should be observed by this TreeView.
 	 */
 	void SetTreeItemsSource( const TArray<ItemType>* InItemsSource)
 	{
-		TreeItemsSource = InItemsSource;
+		SetRootItemsSource(InItemsSource);
+	}
+
+	/**
+	 * Set the Root items. The tree will generate widgets to represent these items.
+	 * @param InItemsSource  A pointer to the array of items that should be observed by this TreeView.
+	 */
+	void SetRootItemsSource(const TArray<ItemType>* InItemsSource)
+	{
+		ensureMsgf(InItemsSource, TEXT("The TreeItemsSource is invalid."));
+		if (TreeViewSource == nullptr || !TreeViewSource->IsSame(reinterpret_cast<const void*>(InItemsSource)))
+		{
+			if (InItemsSource)
+			{
+				SetRootItemsSource(MakeUnique<UE::Slate::ItemsSource::FArrayPointer<ItemType>>(InItemsSource));
+			}
+			else
+			{
+				ClearRootItemsSource();
+			}
+		}
+	}
+
+	/**
+	 * Set the RootItemsSource. The tree will generate widgets to represent these items.
+	 * @param InItemsSource  A pointer to the array of items that should be observed by this TreeView.
+	 */
+	void SetRootItemsSource(TSharedRef<UE::Slate::Containers::TObservableArray<ItemType>> InItemsSource)
+	{
+		if (TreeViewSource == nullptr || !TreeViewSource->IsSame(reinterpret_cast<const void*>(&InItemsSource.Get())))
+		{
+			SetRootItemsSource(MakeUnique<UE::Slate::ItemsSource::FSharedObservableArray<ItemType>>(this->SharedThis(this), MoveTemp(InItemsSource)));
+		}
+	}
+
+	/**
+	 * Establishes a new list of root items being observed by the list.
+	 * Wipes all existing state and requests and will fully rebuild on the next tick.
+	 */
+	void SetRootItemsSource(TUniquePtr<UE::Slate::ItemsSource::IItemsSource<ItemType>> Provider)
+	{
+		TreeViewSource = MoveTemp(Provider);
 		RequestTreeRefresh();
 	}
 
+	void ClearRootItemsSource()
+	{
+		SetRootItemsSource(TUniquePtr<UE::Slate::ItemsSource::IItemsSource<ItemType>>());
+	}
+
+	bool HasValidRootItemsSource() const
+	{
+		return TreeViewSource != nullptr;
+	}
+
+	TArrayView<const ItemType> GetRootItems() const
+	{
+		return TreeViewSource ? TreeViewSource->GetItems() : TArrayView<const ItemType>();
+	}
+
 	/**
 	 * Generates a set of items that are currently expanded.
 	 *
@@ -963,6 +1025,7 @@ protected:
 	/** The delegate that is invoked to recursively expand/collapse a tree items children. */
 	FOnSetExpansionRecursive OnSetExpansionRecursive;
 
+	UE_DEPRECATED(5.3, "Protected access to TreeItemsSource is deprecated. Please use GetTreeItems, SetTreeItemsSource or HasValidTreeItemsSource.")
 	/** A pointer to the items being observed by the tree view. */
 	const TArray<ItemType>* TreeItemsSource;		
 		
@@ -985,6 +1048,8 @@ protected:
 	const FTableViewStyle* Style;
 
 private:		
+	/** Pointer to the source data that we are observing */
+	TUniquePtr<UE::Slate::ItemsSource::IItemsSource<ItemType>> TreeViewSource;
 
 	/** true when the LinearizedItems need to be regenerated. */
 	bool bTreeItemsAreDirty = false;
diff --git a/Engine/Source/Runtime/SlateCore/Public/Widgets/DeclarativeSyntaxSupport.h b/Engine/Source/Runtime/SlateCore/Public/Widgets/DeclarativeSyntaxSupport.h
--- a/Engine/Source/Runtime/SlateCore/Public/Widgets/DeclarativeSyntaxSupport.h
+++ b/Engine/Source/Runtime/SlateCore/Public/Widgets/DeclarativeSyntaxSupport.h
@@ -59,11 +59,12 @@ template<typename WidgetType> struct TSlateBaseNamedArgs;
  *    SLATE_END_ARGS()
  */
 
-#define SLATE_BEGIN_ARGS( WidgetType ) \
+#define SLATE_BEGIN_ARGS( InWidgetType ) \
 	public: \
-	struct FArguments : public TSlateBaseNamedArgs<WidgetType> \
+	struct FArguments : public TSlateBaseNamedArgs<InWidgetType> \
 	{ \
 		typedef FArguments WidgetArgsType; \
+		typedef InWidgetType WidgetType; \
 		FORCENOINLINE FArguments()
 
 /**
@@ -669,9 +670,10 @@ struct FSlateBaseNamedArgs
 
 
 /** Base class for named arguments. Provides settings necessary for all widgets. */
-template<typename WidgetType>
+template<typename InWidgetType>
 struct TSlateBaseNamedArgs : public FSlateBaseNamedArgs
 {
+	typedef InWidgetType WidgetType;
 	typedef typename WidgetType::FArguments WidgetArgsType;
 
 	SLATE_PRIVATE_ATTRIBUTE_FUNCTION(FText, ToolTipText)
