[Engine]

diff --git a/Engine/Source/Runtime/AnimGraphRuntime/Public/BoneControllers/AnimNode_SplineIK.h b/Engine/Source/Runtime/AnimGraphRuntime/Public/BoneControllers/AnimNode_SplineIK.h
--- a/Engine/Source/Runtime/AnimGraphRuntime/Public/BoneControllers/AnimNode_SplineIK.h
+++ b/Engine/Source/Runtime/AnimGraphRuntime/Public/BoneControllers/AnimNode_SplineIK.h
@@ -46,7 +46,7 @@ struct FSplineIKCachedBoneData
 };
 
 USTRUCT(BlueprintInternalUseOnly)
-struct ANIMGRAPHRUNTIME_API FAnimNode_SplineIK : public FAnimNode_SkeletalControlBase
+struct FAnimNode_SplineIK : public FAnimNode_SkeletalControlBase
 {
 	GENERATED_BODY()
 
@@ -101,17 +101,17 @@ struct ANIMGRAPHRUNTIME_API FAnimNode_SplineIK : public FAnimNode_SkeletalContro
 	UPROPERTY(EditAnywhere, BlueprintReadWrite, Category = "Parameters", meta = (PinHiddenByDefault))
 	float Offset;
 
-	FAnimNode_SplineIK();
+	ANIMGRAPHRUNTIME_API FAnimNode_SplineIK();
 
 	// FAnimNode_Base interface
-	virtual void GatherDebugData(FNodeDebugData& DebugData) override;
-	virtual void OnInitializeAnimInstance(const FAnimInstanceProxy* InProxy, const UAnimInstance* InAnimInstance) override;
+	ANIMGRAPHRUNTIME_API virtual void GatherDebugData(FNodeDebugData& DebugData) override;
+	ANIMGRAPHRUNTIME_API virtual void OnInitializeAnimInstance(const FAnimInstanceProxy* InProxy, const UAnimInstance* InAnimInstance) override;
 	virtual bool NeedsOnInitializeAnimInstance() const override { return true; }
 	// End of FAnimNode_Base interface
 
 	// FAnimNode_SkeletalControlBase interface
-	virtual void EvaluateSkeletalControl_AnyThread(FComponentSpacePoseContext& Output, TArray<FBoneTransform>& OutBoneTransforms) override;
-	virtual bool IsValidToEvaluate(const USkeleton* Skeleton, const FBoneContainer& RequiredBones) override;
+	ANIMGRAPHRUNTIME_API virtual void EvaluateSkeletalControl_AnyThread(FComponentSpacePoseContext& Output, TArray<FBoneTransform>& OutBoneTransforms) override;
+	ANIMGRAPHRUNTIME_API virtual bool IsValidToEvaluate(const USkeleton* Skeleton, const FBoneContainer& RequiredBones) override;
 	// End of FAnimNode_SkeletalControlBase interface
 
 	/** Read-only access to spline curves */
@@ -121,47 +121,47 @@ struct ANIMGRAPHRUNTIME_API FAnimNode_SplineIK : public FAnimNode_SkeletalContro
 	const FSplineCurves& GetTransformedSplineCurves() const { return TransformedSpline; }
 
 	/** Get transformed spline point (in component space) for the spline */
-	FTransform GetTransformedSplinePoint(int32 TransformIndex) const;
+	ANIMGRAPHRUNTIME_API FTransform GetTransformedSplinePoint(int32 TransformIndex) const;
 
 	/** Get specified handle transform (in component space) for the spline */
-	FTransform GetControlPoint(int32 TransformIndex) const;
+	ANIMGRAPHRUNTIME_API FTransform GetControlPoint(int32 TransformIndex) const;
 
 	/** Set specified handle transform (in component space) for the spline */
-	void SetControlPoint(int32 TransformIndex, const FTransform& InTransform);
+	ANIMGRAPHRUNTIME_API void SetControlPoint(int32 TransformIndex, const FTransform& InTransform);
 
 	/** Set specified handle location (in component space) for the spline */
-	void SetControlPointLocation(int32 TransformIndex, const FVector& InLocation);
+	ANIMGRAPHRUNTIME_API void SetControlPointLocation(int32 TransformIndex, const FVector& InLocation);
 
 	/** Set specified handle rotation (in component space) for the spline */
-	void SetControlPointRotation(int32 TransformIndex, const FQuat& InRotation);
+	ANIMGRAPHRUNTIME_API void SetControlPointRotation(int32 TransformIndex, const FQuat& InRotation);
 
 	/** Set specified handle scale (in component space) for the spline */
-	void SetControlPointScale(int32 TransformIndex, const FVector& InScale);
+	ANIMGRAPHRUNTIME_API void SetControlPointScale(int32 TransformIndex, const FVector& InScale);
 
 	/** Get the number of spline transforms we are using */
 	int32 GetNumControlPoints() const { return ControlPoints.Num(); }
 
 	/** Build bone references & reallocate transforms from the supplied ref skeleton */
-	void GatherBoneReferences(const FReferenceSkeleton& RefSkeleton);
+	ANIMGRAPHRUNTIME_API void GatherBoneReferences(const FReferenceSkeleton& RefSkeleton);
 
 protected:
 	/** Build spline from reference pose */
-	void BuildBoneSpline(const FReferenceSkeleton& RefSkeleton);
+	ANIMGRAPHRUNTIME_API void BuildBoneSpline(const FReferenceSkeleton& RefSkeleton);
 
 protected:
 	/** Transform the spline using our control points */
-	void TransformSpline();
+	ANIMGRAPHRUNTIME_API void TransformSpline();
 
 	/** Use our linear approximation to determine the earliest intersection with a sphere */
-	float FindParamAtFirstSphereIntersection(const FVector& InOrigin, float InRadius, int32& StartingLinearIndex);
+	ANIMGRAPHRUNTIME_API float FindParamAtFirstSphereIntersection(const FVector& InOrigin, float InRadius, int32& StartingLinearIndex);
 
 private:
 	// FAnimNode_SkeletalControlBase interface
-	virtual void InitializeBoneReferences(const FBoneContainer& RequiredBones) override;
+	ANIMGRAPHRUNTIME_API virtual void InitializeBoneReferences(const FBoneContainer& RequiredBones) override;
 	// End of FAnimNode_SkeletalControlBase interface
 
 	/** Get the current twist value at the specified spline alpha */
-	float GetTwist(float InAlpha, float TotalSplineAlpha);
+	ANIMGRAPHRUNTIME_API float GetTwist(float InAlpha, float TotalSplineAlpha);
 
 	/** Transformed spline */
 	FSplineCurves TransformedSpline;
diff --git a/Engine/Source/Runtime/AnimGraphRuntime/Public/BoneControllers/AnimNode_SpringBone.h b/Engine/Source/Runtime/AnimGraphRuntime/Public/BoneControllers/AnimNode_SpringBone.h
--- a/Engine/Source/Runtime/AnimGraphRuntime/Public/BoneControllers/AnimNode_SpringBone.h
+++ b/Engine/Source/Runtime/AnimGraphRuntime/Public/BoneControllers/AnimNode_SpringBone.h
@@ -17,7 +17,7 @@ class USkeletalMeshComponent;
  */
 
 USTRUCT(BlueprintInternalUseOnly)
-struct ANIMGRAPHRUNTIME_API FAnimNode_SpringBone : public FAnimNode_SkeletalControlBase
+struct FAnimNode_SpringBone : public FAnimNode_SkeletalControlBase
 {
 	GENERATED_USTRUCT_BODY()
 
@@ -101,24 +101,24 @@ struct ANIMGRAPHRUNTIME_API FAnimNode_SpringBone : public FAnimNode_SkeletalCont
 	uint8 bHadValidStrength : 1;
 
 public:
-	FAnimNode_SpringBone();
+	ANIMGRAPHRUNTIME_API FAnimNode_SpringBone();
 
 	// FAnimNode_Base interface
-	virtual void Initialize_AnyThread(const FAnimationInitializeContext& Context) override;
-	virtual void CacheBones_AnyThread(const FAnimationCacheBonesContext& Context) override;
-	virtual void UpdateInternal(const FAnimationUpdateContext& Context) override;
-	virtual void GatherDebugData(FNodeDebugData& DebugData) override;
+	ANIMGRAPHRUNTIME_API virtual void Initialize_AnyThread(const FAnimationInitializeContext& Context) override;
+	ANIMGRAPHRUNTIME_API virtual void CacheBones_AnyThread(const FAnimationCacheBonesContext& Context) override;
+	ANIMGRAPHRUNTIME_API virtual void UpdateInternal(const FAnimationUpdateContext& Context) override;
+	ANIMGRAPHRUNTIME_API virtual void GatherDebugData(FNodeDebugData& DebugData) override;
 	virtual bool HasPreUpdate() const override { return true; }
-	virtual void PreUpdate(const UAnimInstance* InAnimInstance) override;
+	ANIMGRAPHRUNTIME_API virtual void PreUpdate(const UAnimInstance* InAnimInstance) override;
 	// End of FAnimNode_Base interface
 
 	// FAnimNode_SkeletalControlBase interface
-	virtual void EvaluateSkeletalControl_AnyThread(FComponentSpacePoseContext& Output, TArray<FBoneTransform>& OutBoneTransforms) override;
-	virtual bool IsValidToEvaluate(const USkeleton* Skeleton, const FBoneContainer& RequiredBones) override;
+	ANIMGRAPHRUNTIME_API virtual void EvaluateSkeletalControl_AnyThread(FComponentSpacePoseContext& Output, TArray<FBoneTransform>& OutBoneTransforms) override;
+	ANIMGRAPHRUNTIME_API virtual bool IsValidToEvaluate(const USkeleton* Skeleton, const FBoneContainer& RequiredBones) override;
 	// End of FAnimNode_SkeletalControlBase interface
 
 private:
 	// FAnimNode_SkeletalControlBase interface
-	virtual void InitializeBoneReferences(const FBoneContainer& RequiredBones) override;
+	ANIMGRAPHRUNTIME_API virtual void InitializeBoneReferences(const FBoneContainer& RequiredBones) override;
 	// End of FAnimNode_SkeletalControlBase interface
 };
diff --git a/Engine/Source/Runtime/AnimGraphRuntime/Public/BoneControllers/AnimNode_Trail.h b/Engine/Source/Runtime/AnimGraphRuntime/Public/BoneControllers/AnimNode_Trail.h
--- a/Engine/Source/Runtime/AnimGraphRuntime/Public/BoneControllers/AnimNode_Trail.h
+++ b/Engine/Source/Runtime/AnimGraphRuntime/Public/BoneControllers/AnimNode_Trail.h
@@ -43,7 +43,7 @@ struct FRotationLimit
  */
 
 USTRUCT(BlueprintInternalUseOnly)
-struct ANIMGRAPHRUNTIME_API FAnimNode_Trail : public FAnimNode_SkeletalControlBase
+struct FAnimNode_Trail : public FAnimNode_SkeletalControlBase
 {
 	GENERATED_USTRUCT_BODY()
 
@@ -179,27 +179,27 @@ struct ANIMGRAPHRUNTIME_API FAnimNode_Trail : public FAnimNode_SkeletalControlBa
 	TArray<FColor>	PlaneDebugColors;
 #endif // WITH_EDITORONLY_DATA
 
-	FAnimNode_Trail();
+	ANIMGRAPHRUNTIME_API FAnimNode_Trail();
 
 	// FAnimNode_Base interface
-	virtual void Initialize_AnyThread(const FAnimationInitializeContext& Context) override;
-	virtual void UpdateInternal(const FAnimationUpdateContext& Context) override;
-	virtual void GatherDebugData(FNodeDebugData& DebugData) override;
+	ANIMGRAPHRUNTIME_API virtual void Initialize_AnyThread(const FAnimationInitializeContext& Context) override;
+	ANIMGRAPHRUNTIME_API virtual void UpdateInternal(const FAnimationUpdateContext& Context) override;
+	ANIMGRAPHRUNTIME_API virtual void GatherDebugData(FNodeDebugData& DebugData) override;
 	// End of FAnimNode_Base interface
 
 	// FAnimNode_SkeletalControlBase interface
-	virtual void EvaluateSkeletalControl_AnyThread(FComponentSpacePoseContext& Output, TArray<FBoneTransform>& OutBoneTransforms) override;
-	virtual bool IsValidToEvaluate(const USkeleton* Skeleton, const FBoneContainer& RequiredBones) override;
+	ANIMGRAPHRUNTIME_API virtual void EvaluateSkeletalControl_AnyThread(FComponentSpacePoseContext& Output, TArray<FBoneTransform>& OutBoneTransforms) override;
+	ANIMGRAPHRUNTIME_API virtual bool IsValidToEvaluate(const USkeleton* Skeleton, const FBoneContainer& RequiredBones) override;
 	// End of FAnimNode_SkeletalControlBase interface
 
-	void PostLoad();
+	ANIMGRAPHRUNTIME_API void PostLoad();
 
 #if WITH_EDITOR
-	void EnsureChainSize();
+	ANIMGRAPHRUNTIME_API void EnsureChainSize();
 #endif // WITH_EDITOR
 private:
 	// FAnimNode_SkeletalControlBase interface
-	virtual void InitializeBoneReferences(const FBoneContainer& RequiredBones) override;
+	ANIMGRAPHRUNTIME_API virtual void InitializeBoneReferences(const FBoneContainer& RequiredBones) override;
 	// End of FAnimNode_SkeletalControlBase interface
 
 	FVector GetAlignVector(EAxis::Type AxisOption, bool bInvert);
diff --git a/Engine/Source/Runtime/AnimGraphRuntime/Public/BoneControllers/AnimNode_TwistCorrectiveNode.h b/Engine/Source/Runtime/AnimGraphRuntime/Public/BoneControllers/AnimNode_TwistCorrectiveNode.h
--- a/Engine/Source/Runtime/AnimGraphRuntime/Public/BoneControllers/AnimNode_TwistCorrectiveNode.h
+++ b/Engine/Source/Runtime/AnimGraphRuntime/Public/BoneControllers/AnimNode_TwistCorrectiveNode.h
@@ -33,7 +33,7 @@ struct FReferenceBoneFrame
  * This isn't the twist control node for bone twist
  */
 USTRUCT()
-struct ANIMGRAPHRUNTIME_API FAnimNode_TwistCorrectiveNode : public FAnimNode_SkeletalControlBase
+struct FAnimNode_TwistCorrectiveNode : public FAnimNode_SkeletalControlBase
 {
 	GENERATED_USTRUCT_BODY()
 
@@ -76,27 +76,27 @@ struct ANIMGRAPHRUNTIME_API FAnimNode_TwistCorrectiveNode : public FAnimNode_Ske
 	FName CurveName;
 
 public:
-	FAnimNode_TwistCorrectiveNode();
+	ANIMGRAPHRUNTIME_API FAnimNode_TwistCorrectiveNode();
 
 	// FAnimNode_Base interface
-	virtual void GatherDebugData(FNodeDebugData& DebugData) override;
-	virtual void CacheBones_AnyThread(const FAnimationCacheBonesContext& Context)  override;
-	virtual void Initialize_AnyThread(const FAnimationInitializeContext& Context) override;
+	ANIMGRAPHRUNTIME_API virtual void GatherDebugData(FNodeDebugData& DebugData) override;
+	ANIMGRAPHRUNTIME_API virtual void CacheBones_AnyThread(const FAnimationCacheBonesContext& Context)  override;
+	ANIMGRAPHRUNTIME_API virtual void Initialize_AnyThread(const FAnimationInitializeContext& Context) override;
 	// End of FAnimNode_Base interface
 
 	// FAnimNode_SkeletalControlBase interface
-	virtual void EvaluateComponentSpaceInternal(FComponentSpacePoseContext& Context) override;
-	virtual bool IsValidToEvaluate(const USkeleton* Skeleton, const FBoneContainer& RequiredBones) override;
+	ANIMGRAPHRUNTIME_API virtual void EvaluateComponentSpaceInternal(FComponentSpacePoseContext& Context) override;
+	ANIMGRAPHRUNTIME_API virtual bool IsValidToEvaluate(const USkeleton* Skeleton, const FBoneContainer& RequiredBones) override;
 	// End of FAnimNode_SkeletalControlBase interface
 
 	// Type traits support
-	bool Serialize(FArchive& Ar);
-	void PostSerialize(const FArchive& Ar);
+	ANIMGRAPHRUNTIME_API bool Serialize(FArchive& Ar);
+	ANIMGRAPHRUNTIME_API void PostSerialize(const FArchive& Ar);
 
 protected:
 	
 	// FAnimNode_SkeletalControlBase protected interface
-	virtual void InitializeBoneReferences(const FBoneContainer& RequiredBones) override;
+	ANIMGRAPHRUNTIME_API virtual void InitializeBoneReferences(const FBoneContainer& RequiredBones) override;
 
 private:
 	// Reference Pose Angle
@@ -118,4 +118,4 @@ struct TStructOpsTypeTraits<FAnimNode_TwistCorrectiveNode> : public TStructOpsTy
 		WithSerializer = true,
 		WithPostSerialize = true
 	};
-};
\ No newline at end of file
+};
diff --git a/Engine/Source/Runtime/AnimGraphRuntime/Public/BoneControllers/AnimNode_TwoBoneIK.h b/Engine/Source/Runtime/AnimGraphRuntime/Public/BoneControllers/AnimNode_TwoBoneIK.h
--- a/Engine/Source/Runtime/AnimGraphRuntime/Public/BoneControllers/AnimNode_TwoBoneIK.h
+++ b/Engine/Source/Runtime/AnimGraphRuntime/Public/BoneControllers/AnimNode_TwoBoneIK.h
@@ -17,7 +17,7 @@ class USkeletalMeshComponent;
  */
 
 USTRUCT(BlueprintInternalUseOnly)
-struct ANIMGRAPHRUNTIME_API FAnimNode_TwoBoneIK : public FAnimNode_SkeletalControlBase
+struct FAnimNode_TwoBoneIK : public FAnimNode_SkeletalControlBase
 {
 	GENERATED_USTRUCT_BODY()
 	
@@ -99,24 +99,24 @@ struct ANIMGRAPHRUNTIME_API FAnimNode_TwoBoneIK : public FAnimNode_SkeletalContr
 	UPROPERTY(EditAnywhere, Category = IK)
 	uint8 bAllowTwist : 1;
 
-	FAnimNode_TwoBoneIK();
+	ANIMGRAPHRUNTIME_API FAnimNode_TwoBoneIK();
 
 	// FAnimNode_Base interface
-	virtual void GatherDebugData(FNodeDebugData& DebugData) override;
-	virtual void Initialize_AnyThread(const FAnimationInitializeContext& Context) override;
+	ANIMGRAPHRUNTIME_API virtual void GatherDebugData(FNodeDebugData& DebugData) override;
+	ANIMGRAPHRUNTIME_API virtual void Initialize_AnyThread(const FAnimationInitializeContext& Context) override;
 	// End of FAnimNode_Base interface
 
 	// FAnimNode_SkeletalControlBase interface
-	virtual void EvaluateSkeletalControl_AnyThread(FComponentSpacePoseContext& Output, TArray<FBoneTransform>& OutBoneTransforms) override;
-	virtual bool IsValidToEvaluate(const USkeleton* Skeleton, const FBoneContainer& RequiredBones) override;
+	ANIMGRAPHRUNTIME_API virtual void EvaluateSkeletalControl_AnyThread(FComponentSpacePoseContext& Output, TArray<FBoneTransform>& OutBoneTransforms) override;
+	ANIMGRAPHRUNTIME_API virtual bool IsValidToEvaluate(const USkeleton* Skeleton, const FBoneContainer& RequiredBones) override;
 #if WITH_EDITOR
-	void ConditionalDebugDraw(FPrimitiveDrawInterface* PDI, USkeletalMeshComponent* MeshComp) const;
+	ANIMGRAPHRUNTIME_API void ConditionalDebugDraw(FPrimitiveDrawInterface* PDI, USkeletalMeshComponent* MeshComp) const;
 #endif // WITH_EDITOR
 	// End of FAnimNode_SkeletalControlBase interface
-	static FTransform GetTargetTransform(const FTransform& InComponentTransform, FCSPose<FCompactPose>& MeshBases, FBoneSocketTarget& InTarget, EBoneControlSpace Space, const FVector& InOffset);
+	static ANIMGRAPHRUNTIME_API FTransform GetTargetTransform(const FTransform& InComponentTransform, FCSPose<FCompactPose>& MeshBases, FBoneSocketTarget& InTarget, EBoneControlSpace Space, const FVector& InOffset);
 private:
 	// FAnimNode_SkeletalControlBase interface
-	virtual void InitializeBoneReferences(const FBoneContainer& RequiredBones) override;
+	ANIMGRAPHRUNTIME_API virtual void InitializeBoneReferences(const FBoneContainer& RequiredBones) override;
 	// End of FAnimNode_SkeletalControlBase interface
 
 #if WITH_EDITOR
diff --git a/Engine/Source/Runtime/AnimGraphRuntime/Public/BoneControllers/BoneControllerSolvers.h b/Engine/Source/Runtime/AnimGraphRuntime/Public/BoneControllers/BoneControllerSolvers.h
--- a/Engine/Source/Runtime/AnimGraphRuntime/Public/BoneControllers/BoneControllerSolvers.h
+++ b/Engine/Source/Runtime/AnimGraphRuntime/Public/BoneControllers/BoneControllerSolvers.h
@@ -11,7 +11,7 @@
 #include "BoneControllerSolvers.generated.h"
 
 USTRUCT(BlueprintInternalUseOnly)
-struct ANIMGRAPHRUNTIME_API FIKFootPelvisPullDownSolver
+struct FIKFootPelvisPullDownSolver
 {
 	GENERATED_BODY()
 
@@ -44,5 +44,5 @@ struct ANIMGRAPHRUNTIME_API FIKFootPelvisPullDownSolver
 	int32 PelvisAdjustmentMaxIter = 3;
 
 	// Iteratively pulls the character pelvis towards the ground based on the relationship of driven IK foot targets versus FK foot limits
-	FTransform Solve(FTransform PelvisTransform, TArrayView<const float> FKFootDistancesToPelvis, TArrayView<const FVector> IKFootLocations, float DeltaTime);
-};
\ No newline at end of file
+	ANIMGRAPHRUNTIME_API FTransform Solve(FTransform PelvisTransform, TArrayView<const float> FKFootDistancesToPelvis, TArrayView<const FVector> IKFootLocations, float DeltaTime);
+};
diff --git a/Engine/Source/Runtime/AnimGraphRuntime/Public/BoneControllers/BoneControllerTypes.h b/Engine/Source/Runtime/AnimGraphRuntime/Public/BoneControllers/BoneControllerTypes.h
--- a/Engine/Source/Runtime/AnimGraphRuntime/Public/BoneControllers/BoneControllerTypes.h
+++ b/Engine/Source/Runtime/AnimGraphRuntime/Public/BoneControllers/BoneControllerTypes.h
@@ -36,7 +36,7 @@ enum class EWarpingVectorMode : uint8
 
 // Vector values which may be specified in a configured space
 USTRUCT(BlueprintType)
-struct ANIMGRAPHRUNTIME_API FWarpingVectorValue
+struct FWarpingVectorValue
 {
 	GENERATED_BODY()
 
@@ -49,5 +49,5 @@ struct ANIMGRAPHRUNTIME_API FWarpingVectorValue
 	FVector Value = FVector::ZeroVector;
 
 	// Retrieves a normalized Component-space direction from the specified DirectionMode and Direction value
-	FVector AsComponentSpaceDirection(const FAnimInstanceProxy* AnimInstanceProxy, const FTransform& IKFootRootTransform) const;
-};
\ No newline at end of file
+	ANIMGRAPHRUNTIME_API FVector AsComponentSpaceDirection(const FAnimInstanceProxy* AnimInstanceProxy, const FTransform& IKFootRootTransform) const;
+};
diff --git a/Engine/Source/Runtime/AnimGraphRuntime/Public/KismetAnimationLibrary.h b/Engine/Source/Runtime/AnimGraphRuntime/Public/KismetAnimationLibrary.h
--- a/Engine/Source/Runtime/AnimGraphRuntime/Public/KismetAnimationLibrary.h
+++ b/Engine/Source/Runtime/AnimGraphRuntime/Public/KismetAnimationLibrary.h
@@ -14,8 +14,8 @@ class USkeletalMeshComponent;
 /**
  *	A library of the most common animation blueprint functions.
  */
-UCLASS(meta=(ScriptName="AnimGraphLibrary", BlueprintThreadSafe, DocumentationPolicy="Strict"))
-class ANIMGRAPHRUNTIME_API UKismetAnimationLibrary : public UBlueprintFunctionLibrary
+UCLASS(meta=(ScriptName="AnimGraphLibrary", BlueprintThreadSafe, DocumentationPolicy="Strict"), MinimalAPI)
+class UKismetAnimationLibrary : public UBlueprintFunctionLibrary
 {
 	GENERATED_UCLASS_BODY()
 
@@ -34,7 +34,7 @@ class ANIMGRAPHRUNTIME_API UKismetAnimationLibrary : public UBlueprintFunctionLi
 	 *	@param MaxStretchScale The maximum multiplier for the stretch to reach.
 	 */
 	UFUNCTION(BlueprintPure, Category = "Animation|Utilities", meta = (DisplayName = "Two Bone IK Function", ScriptName = "TwoBoneIK", bAllowStretching = "false", StartStretchRatio = "1.0", MaxStretchScale = "1.2"))
-	static void K2_TwoBoneIK(
+	static ANIMGRAPHRUNTIME_API void K2_TwoBoneIK(
 		const FVector& RootPos, 
 		const FVector& JointPos, 
 		const FVector& EndPos, 
@@ -58,7 +58,7 @@ class ANIMGRAPHRUNTIME_API UKismetAnimationLibrary : public UBlueprintFunctionLi
 	 *	@param ClampConeInDegree A limit for only allowing the lookat to rotate as much as defined by the float value
 	 */
 	UFUNCTION(BlueprintPure, Category = "Animation|Utilities", meta = (DisplayName = "Look At Function", ScriptName = "LookAt", bUseUpVector = "false"))
-	static FTransform K2_LookAt(
+	static ANIMGRAPHRUNTIME_API FTransform K2_LookAt(
 		const FTransform& CurrentTransform, 
 		const FVector& TargetPosition, 
 		FVector LookAtVector, 
@@ -82,7 +82,7 @@ class ANIMGRAPHRUNTIME_API UKismetAnimationLibrary : public UBlueprintFunctionLi
 	 *	@param OutRangeMax The maximum for the output range (commonly == 1.0)
 	 */
 	UFUNCTION(BlueprintPure, Category = "Animation|Utilities", meta = (NotBlueprintThreadSafe, DisplayName = "Get Distance Between Two Sockets", ScriptName = "DistanceBetweenSockets", bRemapRange= "false"))
-	static float K2_DistanceBetweenTwoSocketsAndMapRange(
+	static ANIMGRAPHRUNTIME_API float K2_DistanceBetweenTwoSocketsAndMapRange(
 		const USkeletalMeshComponent* Component, 
 		const FName SocketOrBoneNameA, 
 		ERelativeTransformSpace SocketSpaceA, 
@@ -102,7 +102,7 @@ class ANIMGRAPHRUNTIME_API UKismetAnimationLibrary : public UBlueprintFunctionLi
 	 *	@param SocketOrBoneNameTo The name of the second socket / bone
 	 */
 	UFUNCTION(BlueprintPure, Category = "Animation|Utilities", meta = (NotBlueprintThreadSafe, DisplayName = "Get Direction Between Sockets", ScriptName = "DirectionBetweenSockets"))
-	static FVector K2_DirectionBetweenSockets(const USkeletalMeshComponent* Component, const FName SocketOrBoneNameFrom, const FName SocketOrBoneNameTo);
+	static ANIMGRAPHRUNTIME_API FVector K2_DirectionBetweenSockets(const USkeletalMeshComponent* Component, const FName SocketOrBoneNameFrom, const FName SocketOrBoneNameTo);
 
 	/** 
 	 *	This function creates perlin noise from input X, Y, Z, and then range map to RangeOut, and out put to OutX, OutY, OutZ
@@ -118,7 +118,7 @@ class ANIMGRAPHRUNTIME_API UKismetAnimationLibrary : public UBlueprintFunctionLi
 	 *	@param RangeOutMaxZ The maximum for the output range for the z component
 	 */
 	UFUNCTION(BlueprintPure, Category = "Animation|Utilities", meta = (DisplayName = "Make Perlin Noise Vector and Remap", ScriptName = "MakeVectorFromPerlinNoise", RangeOutMinX = "-1.f", RangeOutMaxX = "1.f", RangeOutMinY = "-1.f", RangeOutMaxY = "1.f", RangeOutMinZ = "-1.f", RangeOutMaxZ = "1.f"))
-	static FVector K2_MakePerlinNoiseVectorAndRemap(
+	static ANIMGRAPHRUNTIME_API FVector K2_MakePerlinNoiseVectorAndRemap(
 		float X,
 		float Y,
 		float Z,
@@ -138,7 +138,7 @@ class ANIMGRAPHRUNTIME_API UKismetAnimationLibrary : public UBlueprintFunctionLi
 	 *	@param RangeOutMax The maximum for the output range
 	 */
 	UFUNCTION(BlueprintPure, Category = "Animation|Utilities", meta = (DisplayName = "Make Perlin Noise and Remap", ScriptName = "MakeFloatFromPerlinNoise", RangeOutMinX = "-1.f", RangeOutMaxX = "1.f", RangeOutMinY = "-1.f", RangeOutMaxY = "1.f", RangeOutMinZ = "-1.f", RangeOutMaxZ = "1.f"))
-	static float K2_MakePerlinNoiseAndRemap(
+	static ANIMGRAPHRUNTIME_API float K2_MakePerlinNoiseAndRemap(
 		float Value,
 		float RangeOutMin,
 		float RangeOutMax
@@ -156,7 +156,7 @@ class ANIMGRAPHRUNTIME_API UKismetAnimationLibrary : public UBlueprintFunctionLi
 	 *	@param VelocityMax The maximum velocity to use for normalization (if both min and max are set to 0, normalization is turned off)
 	 */
 	UFUNCTION(BlueprintCallable, Category = "Animation|Utilities", meta = (DisplayName = "Calculate Velocity From Position History", ScriptName = "CalculateVelocityFromPositionHistory", NumberOfSamples = "16", VelocityMin = "0.f", VelocityMax = "128.f"))
-	static float K2_CalculateVelocityFromPositionHistory(
+	static ANIMGRAPHRUNTIME_API float K2_CalculateVelocityFromPositionHistory(
 		float DeltaSeconds,
 		FVector Position,
 		UPARAM(ref) FPositionHistory& History,
@@ -184,7 +184,7 @@ class ANIMGRAPHRUNTIME_API UKismetAnimationLibrary : public UBlueprintFunctionLi
 	 *	@param CustomCurve The curve to use if the easing type is "Custom"
 	 */
 	UFUNCTION(BlueprintCallable, Category = "Animation|Utilities", meta = (NotBlueprintThreadSafe, DisplayName = "Calculate Velocity From Sockets", ScriptName = "CalculateVelocityFromSockets", NumberOfSamples = "16", VelocityMin = "0.f", VelocityMax = "128.f"))
-	static float K2_CalculateVelocityFromSockets(
+	static ANIMGRAPHRUNTIME_API float K2_CalculateVelocityFromSockets(
 		float DeltaSeconds,
 		USkeletalMeshComponent * Component,
 		const FName SocketOrBoneName,
@@ -203,7 +203,7 @@ class ANIMGRAPHRUNTIME_API UKismetAnimationLibrary : public UBlueprintFunctionLi
 	 *	This function starts measuring the time for a profiling bracket
 	 */
 	UFUNCTION(BlueprintCallable, Category = "Animation|Utilities", meta = (DisplayName = "Start Profiling Timer"))
-	static void K2_StartProfilingTimer();
+	static ANIMGRAPHRUNTIME_API void K2_StartProfilingTimer();
 
 	/**
 	 *	This function ends measuring a profiling bracket and optionally logs the result
@@ -213,7 +213,7 @@ class ANIMGRAPHRUNTIME_API UKismetAnimationLibrary : public UBlueprintFunctionLi
 	 *  @result The time spent in milliseconds
 	 */
 	UFUNCTION(BlueprintCallable, Category = "Animation|Utilities", meta = (DisplayName = "End Profiling Timer"))
-	static float K2_EndProfilingTimer(bool bLog = true, const FString& LogPrefix = "");
+	static ANIMGRAPHRUNTIME_API float K2_EndProfilingTimer(bool bLog = true, const FString& LogPrefix = "");
 
 	/** 
 	 * Returns degree of the angle between Velocity and Rotation forward vector
@@ -222,6 +222,6 @@ class ANIMGRAPHRUNTIME_API UKismetAnimationLibrary : public UBlueprintFunctionLi
 	 * @param	BaseRotation	The base rotation, e.g. of a pawn
 	 */
 	UFUNCTION(BlueprintPure, Category="Animation|Utilities")
-	static float CalculateDirection(const FVector& Velocity, const FRotator& BaseRotation);
+	static ANIMGRAPHRUNTIME_API float CalculateDirection(const FVector& Velocity, const FRotator& BaseRotation);
 };
 
diff --git a/Engine/Source/Runtime/AnimGraphRuntime/Public/LayeredBoneBlendLibrary.h b/Engine/Source/Runtime/AnimGraphRuntime/Public/LayeredBoneBlendLibrary.h
--- a/Engine/Source/Runtime/AnimGraphRuntime/Public/LayeredBoneBlendLibrary.h
+++ b/Engine/Source/Runtime/AnimGraphRuntime/Public/LayeredBoneBlendLibrary.h
@@ -1,4 +1,4 @@
-﻿// Copyright Epic Games, Inc. All Rights Reserved.
+// Copyright Epic Games, Inc. All Rights Reserved.
 
 #pragma once
 
@@ -21,15 +21,15 @@ struct FLayeredBoneBlendReference : public FAnimNodeReference
 /**
  * Exposes operations to be performed on a layered bone blend anim node.
  */
-UCLASS()
-class ANIMGRAPHRUNTIME_API ULayeredBoneBlendLibrary : public UBlueprintFunctionLibrary
+UCLASS(MinimalAPI)
+class ULayeredBoneBlendLibrary : public UBlueprintFunctionLibrary
 {
 	GENERATED_BODY()
 
 public:
 	/** Get a layered bone blend context from an anim node context. */
 	UFUNCTION(BlueprintCallable, Category = "Layered Bone Blend", meta = (BlueprintThreadSafe, ExpandEnumAsExecs = "Result"))
-	static FLayeredBoneBlendReference ConvertToLayeredBoneBlend(const FAnimNodeReference& Node, EAnimNodeReferenceConversionResult& Result);
+	static ANIMGRAPHRUNTIME_API FLayeredBoneBlendReference ConvertToLayeredBoneBlend(const FAnimNodeReference& Node, EAnimNodeReferenceConversionResult& Result);
 
 	/** Get a layered bone blend context from an anim node context (pure). */
 	UFUNCTION(BlueprintPure, Category = "Layered Bone Blend", meta = (BlueprintThreadSafe, DisplayName = "Convert to Layered Bone Blend"))
@@ -42,7 +42,7 @@ public:
 
 	/** Get the number of poses that a layered bone blend node has (this does not include the base pose) */
 	UFUNCTION(BlueprintPure, Category = "Layered Bone Blend", meta = (BlueprintThreadSafe))
-	static int32 GetNumPoses(const FLayeredBoneBlendReference& LayeredBoneBlend);
+	static ANIMGRAPHRUNTIME_API int32 GetNumPoses(const FLayeredBoneBlendReference& LayeredBoneBlend);
 	
 	/**
 	 * Sets the currently-used blend mask for a blended input pose by name.
@@ -52,5 +52,5 @@ public:
 	 * @param	BlendMaskName	The name of the blend mask to use 
 	 */
 	UFUNCTION(BlueprintCallable, Category = "Layered Bone Blend", meta = (BlueprintThreadSafe))
-	static FLayeredBoneBlendReference SetBlendMask(const FAnimUpdateContext& UpdateContext, const FLayeredBoneBlendReference& LayeredBoneBlend, int32 PoseIndex, FName BlendMaskName);
-};
\ No newline at end of file
+	static ANIMGRAPHRUNTIME_API FLayeredBoneBlendReference SetBlendMask(const FAnimUpdateContext& UpdateContext, const FLayeredBoneBlendReference& LayeredBoneBlend, int32 PoseIndex, FName BlendMaskName);
+};
diff --git a/Engine/Source/Runtime/AnimGraphRuntime/Public/LinkedAnimGraphLibrary.h b/Engine/Source/Runtime/AnimGraphRuntime/Public/LinkedAnimGraphLibrary.h
--- a/Engine/Source/Runtime/AnimGraphRuntime/Public/LinkedAnimGraphLibrary.h
+++ b/Engine/Source/Runtime/AnimGraphRuntime/Public/LinkedAnimGraphLibrary.h
@@ -18,15 +18,15 @@ struct FLinkedAnimGraphReference : public FAnimNodeReference
 };
 
 // Exposes operations to be performed on anim node contexts
-UCLASS()
-class ANIMGRAPHRUNTIME_API ULinkedAnimGraphLibrary : public UBlueprintFunctionLibrary
+UCLASS(MinimalAPI)
+class ULinkedAnimGraphLibrary : public UBlueprintFunctionLibrary
 {
 	GENERATED_BODY()
 
 public:
 	/** Get a linked anim graph reference from an anim node reference */
 	UFUNCTION(BlueprintCallable, Category = "Animation|Linked Anim Graphs", meta=(BlueprintThreadSafe, ExpandEnumAsExecs = "Result"))
-	static FLinkedAnimGraphReference ConvertToLinkedAnimGraph(const FAnimNodeReference& Node, EAnimNodeReferenceConversionResult& Result);
+	static ANIMGRAPHRUNTIME_API FLinkedAnimGraphReference ConvertToLinkedAnimGraph(const FAnimNodeReference& Node, EAnimNodeReferenceConversionResult& Result);
 
 	/** Get a linked anim graph reference from an anim node reference (pure) */
 	UFUNCTION(BlueprintPure, Category = "Animation|Linked Anim Graphs", meta=(BlueprintThreadSafe, DisplayName = "Convert to Linked Anim Graph"))
@@ -39,9 +39,9 @@ public:
 	
 	/** Returns whether the node hosts an instance (e.g. linked anim graph or layer) */
 	UFUNCTION(BlueprintPure, Category = "Animation|Linked Anim Graphs", meta=(BlueprintThreadSafe))
-	static bool HasLinkedAnimInstance(const FLinkedAnimGraphReference& Node);
+	static ANIMGRAPHRUNTIME_API bool HasLinkedAnimInstance(const FLinkedAnimGraphReference& Node);
 
 	/** Get the linked instance is hosted by this node. If the node does not host an instance then HasLinkedAnimInstance will return false */
 	UFUNCTION(BlueprintPure, Category = "Animation|Linked Anim Graphs", meta=(BlueprintThreadSafe))
-	static UAnimInstance* GetLinkedAnimInstance(const FLinkedAnimGraphReference& Node);
-};
\ No newline at end of file
+	static ANIMGRAPHRUNTIME_API UAnimInstance* GetLinkedAnimInstance(const FLinkedAnimGraphReference& Node);
+};
diff --git a/Engine/Source/Runtime/AnimGraphRuntime/Public/PlayMontageCallbackProxy.h b/Engine/Source/Runtime/AnimGraphRuntime/Public/PlayMontageCallbackProxy.h
--- a/Engine/Source/Runtime/AnimGraphRuntime/Public/PlayMontageCallbackProxy.h
+++ b/Engine/Source/Runtime/AnimGraphRuntime/Public/PlayMontageCallbackProxy.h
@@ -10,8 +10,8 @@
 
 DECLARE_DYNAMIC_MULTICAST_DELEGATE_OneParam(FOnMontagePlayDelegate, FName, NotifyName);
 
-UCLASS()
-class ANIMGRAPHRUNTIME_API UPlayMontageCallbackProxy : public UObject
+UCLASS(MinimalAPI)
+class UPlayMontageCallbackProxy : public UObject
 {
 	GENERATED_UCLASS_BODY()
 
@@ -35,7 +35,7 @@ class ANIMGRAPHRUNTIME_API UPlayMontageCallbackProxy : public UObject
 
 	// Called to perform the query internally
 	UFUNCTION(BlueprintCallable, meta = (BlueprintInternalUseOnly = "true"))
-	static UPlayMontageCallbackProxy* CreateProxyObjectForPlayMontage(
+	static ANIMGRAPHRUNTIME_API UPlayMontageCallbackProxy* CreateProxyObjectForPlayMontage(
 		class USkeletalMeshComponent* InSkeletalMeshComponent, 
 		class UAnimMontage* MontageToPlay, 
 		float PlayRate = 1.f, 
@@ -44,21 +44,21 @@ class ANIMGRAPHRUNTIME_API UPlayMontageCallbackProxy : public UObject
 
 public:
 	//~ Begin UObject Interface
-	virtual void BeginDestroy() override;
+	ANIMGRAPHRUNTIME_API virtual void BeginDestroy() override;
 	//~ End UObject Interface
 
 protected:
 	UFUNCTION()
-	void OnMontageBlendingOut(UAnimMontage* Montage, bool bInterrupted);
+	ANIMGRAPHRUNTIME_API void OnMontageBlendingOut(UAnimMontage* Montage, bool bInterrupted);
 
 	UFUNCTION()
-	void OnMontageEnded(UAnimMontage* Montage, bool bInterrupted);
+	ANIMGRAPHRUNTIME_API void OnMontageEnded(UAnimMontage* Montage, bool bInterrupted);
 
 	UFUNCTION()
-	void OnNotifyBeginReceived(FName NotifyName, const FBranchingPointNotifyPayload& BranchingPointNotifyPayload);
+	ANIMGRAPHRUNTIME_API void OnNotifyBeginReceived(FName NotifyName, const FBranchingPointNotifyPayload& BranchingPointNotifyPayload);
 
 	UFUNCTION()
-	void OnNotifyEndReceived(FName NotifyName, const FBranchingPointNotifyPayload& BranchingPointNotifyPayload);
+	ANIMGRAPHRUNTIME_API void OnNotifyEndReceived(FName NotifyName, const FBranchingPointNotifyPayload& BranchingPointNotifyPayload);
 
 private:
 	TWeakObjectPtr<UAnimInstance> AnimInstancePtr;
@@ -73,7 +73,7 @@ private:
 
 protected:
 	// Attempts to play a montage with the specified settings. Returns whether it started or not.
-	bool PlayMontage(
+	ANIMGRAPHRUNTIME_API bool PlayMontage(
 		class USkeletalMeshComponent* InSkeletalMeshComponent,
 		class UAnimMontage* MontageToPlay,
 		float PlayRate = 1.f,
diff --git a/Engine/Source/Runtime/AnimGraphRuntime/Public/RBF/RBFSolver.h b/Engine/Source/Runtime/AnimGraphRuntime/Public/RBF/RBFSolver.h
--- a/Engine/Source/Runtime/AnimGraphRuntime/Public/RBF/RBFSolver.h
+++ b/Engine/Source/Runtime/AnimGraphRuntime/Public/RBF/RBFSolver.h
@@ -108,7 +108,7 @@ enum class ERBFNormalizeMethod : uint8
 
 /** Struct storing a particular entry within the RBF */
 USTRUCT()
-struct ANIMGRAPHRUNTIME_API FRBFEntry
+struct FRBFEntry
 {
 	GENERATED_BODY()
 
@@ -117,18 +117,18 @@ struct ANIMGRAPHRUNTIME_API FRBFEntry
 	TArray<float> Values;
 
 	/** Return a target as an rotator, assuming Values is a sequence of Euler entries. Index is which Euler to convert.*/
-	FRotator AsRotator(int32 Index) const;
+	ANIMGRAPHRUNTIME_API FRotator AsRotator(int32 Index) const;
 
 	/** Return a target as a quaternion, assuming Values is a sequence of Euler entries. Index is which Euler to convert. */
-	FQuat AsQuat(int32 Index) const;
+	ANIMGRAPHRUNTIME_API FQuat AsQuat(int32 Index) const;
 
-	FVector AsVector(int32 Index) const;
+	ANIMGRAPHRUNTIME_API FVector AsVector(int32 Index) const;
 
 
 	/** Set this entry to 3 floats from supplied rotator */
-	void AddFromRotator(const FRotator& InRot);
+	ANIMGRAPHRUNTIME_API void AddFromRotator(const FRotator& InRot);
 	/** Set this entry to 3 floats from supplied vector */
-	void AddFromVector(const FVector& InVector);
+	ANIMGRAPHRUNTIME_API void AddFromVector(const FVector& InVector);
 
 	/** Return dimensionality of this target */
 	int32 GetDimensions() const
@@ -139,7 +139,7 @@ struct ANIMGRAPHRUNTIME_API FRBFEntry
 
 /** Data about a particular target in the RBF, including scaling factor */
 USTRUCT()
-struct ANIMGRAPHRUNTIME_API FRBFTarget : public FRBFEntry
+struct FRBFTarget : public FRBFEntry
 {
 	GENERATED_BODY()
 
@@ -177,7 +177,7 @@ struct ANIMGRAPHRUNTIME_API FRBFTarget : public FRBFEntry
 };
 
 /** Struct for storing RBF results - target index and corresponding weight */
-struct ANIMGRAPHRUNTIME_API FRBFOutputWeight
+struct FRBFOutputWeight
 {
 	/** Index of target */
 	int32 TargetIndex;
@@ -197,7 +197,7 @@ struct ANIMGRAPHRUNTIME_API FRBFOutputWeight
 
 /** Parameters used by RBF solver */
 USTRUCT(BlueprintType)
-struct ANIMGRAPHRUNTIME_API FRBFParams
+struct FRBFParams
 {
 	GENERATED_BODY()
 
@@ -252,38 +252,38 @@ struct ANIMGRAPHRUNTIME_API FRBFParams
 	UPROPERTY(EditAnywhere, BlueprintReadWrite, Category = RBFData, meta = (UIMin = "0", UIMax = "90", EditCondition = "NormalizeMethod == ERBFNormalizeMethod::NormalizeWithinMedian"))
 	float MedianMax;
 
-	FRBFParams();
+	ANIMGRAPHRUNTIME_API FRBFParams();
 
 	/** Util for returning unit direction vector for swing axis */
-	FVector GetTwistAxisVector() const;
+	ANIMGRAPHRUNTIME_API FVector GetTwistAxisVector() const;
 };
 
 struct ANIMGRAPHRUNTIME_API FRBFSolverData;
 
 /** Library of Radial Basis Function solver functions */
-struct ANIMGRAPHRUNTIME_API FRBFSolver
+struct FRBFSolver
 {
 	/** Given a list of targets, verify which ones are valid for solving the RBF setup. This is mostly about removing identical targets
 		which invalidates the interpolative solver. Returns true if all targets are valid. */
-	static bool ValidateTargets(const FRBFParams& Params, const TArray<FRBFTarget>& Targets, TArray<int>& InvalidTargets);
+	static ANIMGRAPHRUNTIME_API bool ValidateTargets(const FRBFParams& Params, const TArray<FRBFTarget>& Targets, TArray<int>& InvalidTargets);
 
 	/** Given a set of targets and new input entry, give list of activated targets with weights */
-	static TSharedPtr<const FRBFSolverData> InitSolver(const FRBFParams& Params, const TArray<FRBFTarget>& Targets);
+	static ANIMGRAPHRUNTIME_API TSharedPtr<const FRBFSolverData> InitSolver(const FRBFParams& Params, const TArray<FRBFTarget>& Targets);
 
-	static bool IsSolverDataValid(const FRBFSolverData& SolverData, const FRBFParams& Params, const TArray<FRBFTarget>& Targets);
+	static ANIMGRAPHRUNTIME_API bool IsSolverDataValid(const FRBFSolverData& SolverData, const FRBFParams& Params, const TArray<FRBFTarget>& Targets);
 
 	/** Given a set of targets and new input entry, give list of activated targets with weights */
-	static void Solve(const FRBFSolverData& SolverData, const FRBFParams& Params, const TArray<FRBFTarget>& Targets, const FRBFEntry& Input, TArray<FRBFOutputWeight>& OutputWeights);
+	static ANIMGRAPHRUNTIME_API void Solve(const FRBFSolverData& SolverData, const FRBFParams& Params, const TArray<FRBFTarget>& Targets, const FRBFEntry& Input, TArray<FRBFOutputWeight>& OutputWeights);
 
 	/** Util to find distance to nearest neighbour target for each target */
-	static bool FindTargetNeighbourDistances(const FRBFParams& Params, const TArray<FRBFTarget>& Targets, TArray<float>& NeighbourDists);
+	static ANIMGRAPHRUNTIME_API bool FindTargetNeighbourDistances(const FRBFParams& Params, const TArray<FRBFTarget>& Targets, TArray<float>& NeighbourDists);
 
 	/** Util to find distance between two entries, using provided params */
-	static float FindDistanceBetweenEntries(const FRBFEntry& A, const FRBFEntry& B, const FRBFParams& Params, ERBFDistanceMethod OverrideMethod = ERBFDistanceMethod::DefaultMethod);
+	static ANIMGRAPHRUNTIME_API float FindDistanceBetweenEntries(const FRBFEntry& A, const FRBFEntry& B, const FRBFParams& Params, ERBFDistanceMethod OverrideMethod = ERBFDistanceMethod::DefaultMethod);
 
 	/** Returns the radius for a given target */
-	static float GetRadiusForTarget(const FRBFTarget& Target, const FRBFParams& Params);
+	static ANIMGRAPHRUNTIME_API float GetRadiusForTarget(const FRBFTarget& Target, const FRBFParams& Params);
 
 	/** Compute the optimal radius for the given targets. Returns the radius */
-	static float GetOptimalRadiusForTargets(const FRBFParams& Params, const TArray<FRBFTarget>& Targets);
+	static ANIMGRAPHRUNTIME_API float GetOptimalRadiusForTargets(const FRBFParams& Params, const TArray<FRBFTarget>& Targets);
 };
diff --git a/Engine/Source/Runtime/AnimGraphRuntime/Public/SequenceEvaluatorLibrary.h b/Engine/Source/Runtime/AnimGraphRuntime/Public/SequenceEvaluatorLibrary.h
--- a/Engine/Source/Runtime/AnimGraphRuntime/Public/SequenceEvaluatorLibrary.h
+++ b/Engine/Source/Runtime/AnimGraphRuntime/Public/SequenceEvaluatorLibrary.h
@@ -21,15 +21,15 @@ struct FSequenceEvaluatorReference : public FAnimNodeReference
 
 // Exposes operations to be performed on a sequence evaluator anim node
 // Note: Experimental and subject to change!
-UCLASS(Experimental)
-class ANIMGRAPHRUNTIME_API USequenceEvaluatorLibrary : public UBlueprintFunctionLibrary
+UCLASS(Experimental, MinimalAPI)
+class USequenceEvaluatorLibrary : public UBlueprintFunctionLibrary
 {
 	GENERATED_BODY()
 
 public:
 	/** Get a sequence evaluator context from an anim node context */
 	UFUNCTION(BlueprintCallable, Category = "Animation|Sequences", meta=(BlueprintThreadSafe, ExpandEnumAsExecs = "Result"))
-	static FSequenceEvaluatorReference ConvertToSequenceEvaluator(const FAnimNodeReference& Node, EAnimNodeReferenceConversionResult& Result);
+	static ANIMGRAPHRUNTIME_API FSequenceEvaluatorReference ConvertToSequenceEvaluator(const FAnimNodeReference& Node, EAnimNodeReferenceConversionResult& Result);
 
 	/** Get a sequence evaluator context from an anim node context (pure) */
 	UFUNCTION(BlueprintPure, Category = "Animation|Sequences", meta=(BlueprintThreadSafe, DisplayName = "Convert to Sequence Evaluator"))
@@ -42,25 +42,25 @@ public:
 	
 	/** Set the current accumulated time of the sequence evaluator */
 	UFUNCTION(BlueprintCallable, Category = "Animation|Sequences", meta=(BlueprintThreadSafe))
-	static FSequenceEvaluatorReference SetExplicitTime(const FSequenceEvaluatorReference& SequenceEvaluator, float Time);
+	static ANIMGRAPHRUNTIME_API FSequenceEvaluatorReference SetExplicitTime(const FSequenceEvaluatorReference& SequenceEvaluator, float Time);
 
 	/** Advance the current accumulated time of the sequence evaluator */
 	UFUNCTION(BlueprintCallable, Category = "Animation|Sequences", meta = (BlueprintThreadSafe))
-	static FSequenceEvaluatorReference AdvanceTime(const FAnimUpdateContext& UpdateContext, const FSequenceEvaluatorReference& SequenceEvaluator, float PlayRate = 1.0f);
+	static ANIMGRAPHRUNTIME_API FSequenceEvaluatorReference AdvanceTime(const FAnimUpdateContext& UpdateContext, const FSequenceEvaluatorReference& SequenceEvaluator, float PlayRate = 1.0f);
 
 	/** Set the current sequence of the sequence evaluator */
 	UFUNCTION(BlueprintCallable, Category = "Animation|Sequences", meta=(BlueprintThreadSafe))
-	static FSequenceEvaluatorReference SetSequence(const FSequenceEvaluatorReference& SequenceEvaluator, UAnimSequenceBase* Sequence);
+	static ANIMGRAPHRUNTIME_API FSequenceEvaluatorReference SetSequence(const FSequenceEvaluatorReference& SequenceEvaluator, UAnimSequenceBase* Sequence);
 
 	/** Set the current sequence of the sequence evaluator with an inertial blend time */
 	UFUNCTION(BlueprintCallable, Category = "Animation|Sequences", meta=(BlueprintThreadSafe))
-	static FSequenceEvaluatorReference SetSequenceWithInertialBlending(const FAnimUpdateContext& UpdateContext, const FSequenceEvaluatorReference& SequenceEvaluator, UAnimSequenceBase* Sequence, float BlendTime = 0.2f);
+	static ANIMGRAPHRUNTIME_API FSequenceEvaluatorReference SetSequenceWithInertialBlending(const FAnimUpdateContext& UpdateContext, const FSequenceEvaluatorReference& SequenceEvaluator, UAnimSequenceBase* Sequence, float BlendTime = 0.2f);
 	
 	/** Gets the current accumulated time of the sequence evaluator */
 	UFUNCTION(BlueprintPure, Category = "Sequence Evaluator", meta = (BlueprintThreadSafe))
-	static float GetAccumulatedTime(const FSequenceEvaluatorReference& SequenceEvaluator);
+	static ANIMGRAPHRUNTIME_API float GetAccumulatedTime(const FSequenceEvaluatorReference& SequenceEvaluator);
 
 	/** Gets the current sequence of the sequence evaluator */
 	UFUNCTION(BlueprintPure, Category = "Sequence Evaluator", meta = (BlueprintThreadSafe))
-	static UAnimSequenceBase* GetSequence(const FSequenceEvaluatorReference& SequenceEvaluator);
-};
\ No newline at end of file
+	static ANIMGRAPHRUNTIME_API UAnimSequenceBase* GetSequence(const FSequenceEvaluatorReference& SequenceEvaluator);
+};
diff --git a/Engine/Source/Runtime/AnimGraphRuntime/Public/SequencePlayerLibrary.h b/Engine/Source/Runtime/AnimGraphRuntime/Public/SequencePlayerLibrary.h
--- a/Engine/Source/Runtime/AnimGraphRuntime/Public/SequencePlayerLibrary.h
+++ b/Engine/Source/Runtime/AnimGraphRuntime/Public/SequencePlayerLibrary.h
@@ -21,15 +21,15 @@ struct FSequencePlayerReference : public FAnimNodeReference
 
 // Exposes operations to be performed on a sequence player anim node
 // Note: Experimental and subject to change!
-UCLASS(Experimental)
-class ANIMGRAPHRUNTIME_API USequencePlayerLibrary : public UBlueprintFunctionLibrary
+UCLASS(Experimental, MinimalAPI)
+class USequencePlayerLibrary : public UBlueprintFunctionLibrary
 {
 	GENERATED_BODY()
 
 public:
 	/** Get a sequence player context from an anim node context */
 	UFUNCTION(BlueprintCallable, Category = "Animation|Sequences", meta = (BlueprintThreadSafe, ExpandEnumAsExecs = "Result"))
-	static FSequencePlayerReference ConvertToSequencePlayer(const FAnimNodeReference& Node, EAnimNodeReferenceConversionResult& Result);
+	static ANIMGRAPHRUNTIME_API FSequencePlayerReference ConvertToSequencePlayer(const FAnimNodeReference& Node, EAnimNodeReferenceConversionResult& Result);
 
 	/** Get a sequence player context from an anim node context (pure) */
 	UFUNCTION(BlueprintPure, Category = "Animation|Sequences", meta = (BlueprintThreadSafe, DisplayName = "Convert to Sequence Player"))
@@ -42,7 +42,7 @@ public:
 
 	/** Set the current accumulated time of the sequence player */
 	UFUNCTION(BlueprintCallable, Category = "Animation|Sequences", meta = (BlueprintThreadSafe))
-	static FSequencePlayerReference SetAccumulatedTime(const FSequencePlayerReference& SequencePlayer, float Time);
+	static ANIMGRAPHRUNTIME_API FSequencePlayerReference SetAccumulatedTime(const FSequencePlayerReference& SequencePlayer, float Time);
 
 	/** 
 	 * Set the start position of the sequence player. 
@@ -50,45 +50,45 @@ public:
 	 * SetAccumulatedTime to achieve the desired effect of resetting the play time of a sequence player.
 	 */
 	UFUNCTION(BlueprintCallable, Category = "Animation|Sequences", meta = (BlueprintThreadSafe))
-	static FSequencePlayerReference SetStartPosition(const FSequencePlayerReference& SequencePlayer, float StartPosition);
+	static ANIMGRAPHRUNTIME_API FSequencePlayerReference SetStartPosition(const FSequencePlayerReference& SequencePlayer, float StartPosition);
 
 	/** Set the play rate of the sequence player */
 	UFUNCTION(BlueprintCallable, Category = "Animation|Sequences", meta = (BlueprintThreadSafe))
-	static FSequencePlayerReference SetPlayRate(const FSequencePlayerReference& SequencePlayer, float PlayRate);
+	static ANIMGRAPHRUNTIME_API FSequencePlayerReference SetPlayRate(const FSequencePlayerReference& SequencePlayer, float PlayRate);
 
 	/** Set the current sequence of the sequence player */
 	UFUNCTION(BlueprintCallable, Category = "Animation|Sequences", meta = (BlueprintThreadSafe))
-	static FSequencePlayerReference SetSequence(const FSequencePlayerReference& SequencePlayer, UAnimSequenceBase* Sequence);
+	static ANIMGRAPHRUNTIME_API FSequencePlayerReference SetSequence(const FSequencePlayerReference& SequencePlayer, UAnimSequenceBase* Sequence);
 
 	/** Set the current sequence of the sequence player with an inertial blend time */
 	UFUNCTION(BlueprintCallable, Category = "Animation|Sequences", meta = (BlueprintThreadSafe))
-	static FSequencePlayerReference SetSequenceWithInertialBlending(const FAnimUpdateContext& UpdateContext, const FSequencePlayerReference& SequencePlayer, UAnimSequenceBase* Sequence, float BlendTime = 0.2f);
+	static ANIMGRAPHRUNTIME_API FSequencePlayerReference SetSequenceWithInertialBlending(const FAnimUpdateContext& UpdateContext, const FSequencePlayerReference& SequencePlayer, UAnimSequenceBase* Sequence, float BlendTime = 0.2f);
 
 	/** Get the current sequence of the sequence player - DEPRECATED, please use pure version */
 	UFUNCTION(BlueprintCallable, Category = "Animation|Sequences", meta = (BlueprintThreadSafe, DeprecatedFunction))
-	static FSequencePlayerReference GetSequence(const FSequencePlayerReference& SequencePlayer, UPARAM(Ref) UAnimSequenceBase*& SequenceBase);
+	static ANIMGRAPHRUNTIME_API FSequencePlayerReference GetSequence(const FSequencePlayerReference& SequencePlayer, UPARAM(Ref) UAnimSequenceBase*& SequenceBase);
 	
 	/** Get the current sequence of the sequence player */
 	UFUNCTION(BlueprintPure, Category = "Animation|Sequences", meta = (BlueprintThreadSafe, DisplayName = "Get Sequence"))
-	static UAnimSequenceBase* GetSequencePure(const FSequencePlayerReference& SequencePlayer);
+	static ANIMGRAPHRUNTIME_API UAnimSequenceBase* GetSequencePure(const FSequencePlayerReference& SequencePlayer);
 
 	/** Gets the current accumulated time of the sequence player */
 	UFUNCTION(BlueprintPure, Category = "Animation|Sequences", meta = (BlueprintThreadSafe))
-	static float GetAccumulatedTime(const FSequencePlayerReference& SequencePlayer);
+	static ANIMGRAPHRUNTIME_API float GetAccumulatedTime(const FSequencePlayerReference& SequencePlayer);
 
 	/** Get the start position of the sequence player */
 	UFUNCTION(BlueprintPure, Category = "Animation|Sequences", meta = (BlueprintThreadSafe))
-	static float GetStartPosition(const FSequencePlayerReference& SequencePlayer);
+	static ANIMGRAPHRUNTIME_API float GetStartPosition(const FSequencePlayerReference& SequencePlayer);
 
 	/** Get the play rate of the sequence player */
 	UFUNCTION(BlueprintPure, Category = "Animation|Sequences", meta = (BlueprintThreadSafe))
-	static float GetPlayRate(const FSequencePlayerReference& SequencePlayer);
+	static ANIMGRAPHRUNTIME_API float GetPlayRate(const FSequencePlayerReference& SequencePlayer);
 
 	/** Get the looping state of the sequence player */
 	UFUNCTION(BlueprintPure, Category = "Animation|Sequences", meta = (BlueprintThreadSafe))
-	static bool GetLoopAnimation(const FSequencePlayerReference& SequencePlayer);
+	static ANIMGRAPHRUNTIME_API bool GetLoopAnimation(const FSequencePlayerReference& SequencePlayer);
 
 	/** Returns the Play Rate to provide when playing this animation if a specific animation duration is desired */
 	UFUNCTION(BlueprintPure, Category = "Animation|Sequences", meta = (BlueprintThreadSafe))
-	static float ComputePlayRateFromDuration(const FSequencePlayerReference& SequencePlayer, float Duration = 1.f);
-};
\ No newline at end of file
+	static ANIMGRAPHRUNTIME_API float ComputePlayRateFromDuration(const FSequencePlayerReference& SequencePlayer, float Duration = 1.f);
+};
diff --git a/Engine/Source/Runtime/AnimGraphRuntime/Public/SequencerAnimationSupport.h b/Engine/Source/Runtime/AnimGraphRuntime/Public/SequencerAnimationSupport.h
--- a/Engine/Source/Runtime/AnimGraphRuntime/Public/SequencerAnimationSupport.h
+++ b/Engine/Source/Runtime/AnimGraphRuntime/Public/SequencerAnimationSupport.h
@@ -18,13 +18,13 @@ class UAnimInstance;
 class UAnimSequenceBase;
 class UObject;
 
-UINTERFACE(meta = (CannotImplementInterfaceInBlueprint))
-class ANIMGRAPHRUNTIME_API USequencerAnimationSupport : public UInterface
+UINTERFACE(meta = (CannotImplementInterfaceInBlueprint), MinimalAPI)
+class USequencerAnimationSupport : public UInterface
 {
 	GENERATED_UINTERFACE_BODY()
 };
 
-class ANIMGRAPHRUNTIME_API ISequencerAnimationSupport
+class ISequencerAnimationSupport
 {
 	GENERATED_IINTERFACE_BODY()
 
diff --git a/Engine/Source/Runtime/AnimGraphRuntime/Public/SkeletalControlLibrary.h b/Engine/Source/Runtime/AnimGraphRuntime/Public/SkeletalControlLibrary.h
--- a/Engine/Source/Runtime/AnimGraphRuntime/Public/SkeletalControlLibrary.h
+++ b/Engine/Source/Runtime/AnimGraphRuntime/Public/SkeletalControlLibrary.h
@@ -20,15 +20,15 @@ struct FSkeletalControlReference : public FAnimNodeReference
 
 // Exposes operations to be performed on a skeletal control anim node
 // Note: Experimental and subject to change!
-UCLASS(Experimental)
-class ANIMGRAPHRUNTIME_API USkeletalControlLibrary : public UBlueprintFunctionLibrary
+UCLASS(Experimental, MinimalAPI)
+class USkeletalControlLibrary : public UBlueprintFunctionLibrary
 {
 	GENERATED_BODY()
 
 public:
 	/** Get a skeletal control from an anim node */
 	UFUNCTION(BlueprintCallable, Category = "Animation|Skeletal Controls", meta=(BlueprintThreadSafe, ExpandEnumAsExecs = "Result"))
-	static FSkeletalControlReference ConvertToSkeletalControl(const FAnimNodeReference& Node, EAnimNodeReferenceConversionResult& Result);
+	static ANIMGRAPHRUNTIME_API FSkeletalControlReference ConvertToSkeletalControl(const FAnimNodeReference& Node, EAnimNodeReferenceConversionResult& Result);
 
 	/** Get a skeletal control from an anim node (pure) */
 	UFUNCTION(BlueprintPure, Category = "Animation|Skeletal Controls", meta=(BlueprintThreadSafe, DisplayName = "Convert to Skeletal Control"))
@@ -41,9 +41,9 @@ public:
 	
 	/** Set the alpha value of this skeletal control */
 	UFUNCTION(BlueprintCallable, Category = "Animation|Skeletal Controls", meta=(BlueprintThreadSafe))
-	static FSkeletalControlReference SetAlpha(const FSkeletalControlReference& SkeletalControl, float Alpha);
+	static ANIMGRAPHRUNTIME_API FSkeletalControlReference SetAlpha(const FSkeletalControlReference& SkeletalControl, float Alpha);
 
 	/** Get the alpha value of this skeletal control */
 	UFUNCTION(BlueprintPure, Category = "Animation|Skeletal Controls", meta=(BlueprintThreadSafe))
-	static float GetAlpha(const FSkeletalControlReference& SkeletalControl);
-};
\ No newline at end of file
+	static ANIMGRAPHRUNTIME_API float GetAlpha(const FSkeletalControlReference& SkeletalControl);
+};
diff --git a/Engine/Source/Runtime/AnimationCore/Public/AnimationDataSource.h b/Engine/Source/Runtime/AnimationCore/Public/AnimationDataSource.h
--- a/Engine/Source/Runtime/AnimationCore/Public/AnimationDataSource.h
+++ b/Engine/Source/Runtime/AnimationCore/Public/AnimationDataSource.h
@@ -16,8 +16,8 @@
 
 class UClass;
 
-UCLASS()
-class ANIMATIONCORE_API UAnimationDataSourceRegistry : public UObject
+UCLASS(MinimalAPI)
+class UAnimationDataSourceRegistry : public UObject
 {
 	GENERATED_BODY()
 
@@ -31,7 +31,7 @@ public:
 	 * @param InDataSource   The data source to register
 	 * @return true if succeeded
 	 */
-	bool RegisterDataSource(const FName& InName, UObject* InDataSource);
+	ANIMATIONCORE_API bool RegisterDataSource(const FName& InName, UObject* InDataSource);
 
 	/**
 	 * Unregisters / removes an existing data source under a given name.
@@ -40,7 +40,7 @@ public:
 	 * @param InName The name of the data source to remove
 	 * @return true if succeeded
 	 */
-	bool UnregisterDataSource(const FName& InName);
+	ANIMATIONCORE_API bool UnregisterDataSource(const FName& InName);
 
 	/**
 	 * Returns true if this registry contains a source with the given name
@@ -48,7 +48,7 @@ public:
 	 * @param InName The name of the data source to look up.
 	 * @return true if a source with the given name exists
 	 */
-	bool ContainsSource(const FName& InName) const;
+	ANIMATIONCORE_API bool ContainsSource(const FName& InName) const;
 
 	/**
 	 * Returns a given data source and cast it to the expected class.
@@ -57,7 +57,7 @@ public:
 	 * @param InExpectedClass The class expected from the data source
 	 * @return The requested data source
 	 */
-	UObject* RequestSource(const FName& InName, UClass* InExpectedClass) const;
+	ANIMATIONCORE_API UObject* RequestSource(const FName& InName, UClass* InExpectedClass) const;
 
 	/**
 	 * Returns a given data source and cast it to the expected class.
@@ -82,5 +82,5 @@ private:
 	TMap<FName, TWeakObjectPtr<UObject>> DataSources;
 
 	/** Clear Invalid Data Sources that are GC-ed */
-	void ClearInvalidDataSource();
-};
\ No newline at end of file
+	ANIMATIONCORE_API void ClearInvalidDataSource();
+};
diff --git a/Engine/Source/Runtime/AnimationCore/Public/Constraint.h b/Engine/Source/Runtime/AnimationCore/Public/Constraint.h
--- a/Engine/Source/Runtime/AnimationCore/Public/Constraint.h
+++ b/Engine/Source/Runtime/AnimationCore/Public/Constraint.h
@@ -39,7 +39,7 @@ struct FMultiTransformBlendHelper;
  * This is used to filter per axis for constraint options
  */
 USTRUCT(BlueprintType)
-struct ANIMATIONCORE_API FFilterOptionPerAxis
+struct FFilterOptionPerAxis
 {
 	GENERATED_USTRUCT_BODY()
 
@@ -127,7 +127,7 @@ struct ANIMATIONCORE_API FFilterOptionPerAxis
 
 /** A filter for a whole transform */
 USTRUCT(BlueprintType)
-struct ANIMATIONCORE_API FTransformFilter
+struct FTransformFilter
 {
 	GENERATED_BODY()
 
@@ -173,7 +173,7 @@ struct ANIMATIONCORE_API FTransformFilter
 
 /** A description of how to apply a simple transform constraint */
 USTRUCT(BlueprintType)
-struct ANIMATIONCORE_API FConstraintDescription
+struct FConstraintDescription
 {
 	GENERATED_USTRUCT_BODY()
 
@@ -228,7 +228,7 @@ struct ANIMATIONCORE_API FConstraintDescription
  * Used by Constraint for saving the offset, and recovering the offset
  */
 USTRUCT()
-struct ANIMATIONCORE_API FConstraintOffset
+struct FConstraintOffset
 {
 	GENERATED_USTRUCT_BODY()
 
@@ -252,9 +252,9 @@ struct ANIMATIONCORE_API FConstraintOffset
 	{}
 
 	/* Apply the Inverse offset */
-	void ApplyInverseOffset(const FTransform& InTarget, FTransform& OutSource) const;
+	ANIMATIONCORE_API void ApplyInverseOffset(const FTransform& InTarget, FTransform& OutSource) const;
 	/* Save the Inverse offset */
-	void SaveInverseOffset(const FTransform& Source, const FTransform& Target, const FConstraintDescription& Operator);
+	ANIMATIONCORE_API void SaveInverseOffset(const FTransform& Source, const FTransform& Target, const FConstraintDescription& Operator);
 	/** Clear the offset */
 	void Reset()
 	{
@@ -276,7 +276,7 @@ struct ANIMATIONCORE_API FConstraintOffset
 };
 
 USTRUCT(BlueprintType)
-struct ANIMATIONCORE_API FTransformConstraint
+struct FTransformConstraint
 {
 	GENERATED_USTRUCT_BODY()
 
@@ -411,7 +411,7 @@ enum class ETransformConstraintType : uint8
 
 /** A description of how to apply a simple transform constraint */
 USTRUCT()
-struct ANIMATIONCORE_API FTransformConstraintDescription : public FConstraintDescriptionEx
+struct FTransformConstraintDescription : public FConstraintDescriptionEx
 {
 	GENERATED_USTRUCT_BODY()
 
@@ -422,7 +422,7 @@ struct ANIMATIONCORE_API FTransformConstraintDescription : public FConstraintDes
 		:TransformType (InType)
 	{}
 	
-	virtual void AccumulateConstraintTransform(const FTransform& TargetTransform, const FTransform& CurrentTransform, const FTransform& CurrentParentTransform, float Weight, FMultiTransformBlendHelper& BlendHelperInLocalSpace) const override;
+	ANIMATIONCORE_API virtual void AccumulateConstraintTransform(const FTransform& TargetTransform, const FTransform& CurrentTransform, const FTransform& CurrentParentTransform, float Weight, FMultiTransformBlendHelper& BlendHelperInLocalSpace) const override;
 	virtual bool DoesAffectRotation() const override { return TransformType == ETransformConstraintType::Rotation; }
 	virtual bool DoesAffectTranslation() const override { return TransformType == ETransformConstraintType::Translation; }
 	virtual bool DoesAffectScale() const override { return TransformType == ETransformConstraintType::Scale; }
@@ -456,7 +456,7 @@ struct ANIMATIONCORE_API FTransformConstraintDescription : public FConstraintDes
 
 /** A description of how to apply aim constraint */
 USTRUCT()
-struct ANIMATIONCORE_API FAimConstraintDescription : public FConstraintDescriptionEx
+struct FAimConstraintDescription : public FConstraintDescriptionEx
 {
 	GENERATED_USTRUCT_BODY()
 
@@ -477,7 +477,7 @@ struct ANIMATIONCORE_API FAimConstraintDescription : public FConstraintDescripti
 	{
 	}
 
-	virtual void AccumulateConstraintTransform(const FTransform& TargetTransform, const FTransform& CurrentTransform, const FTransform& CurrentParentTransform, float Weight, FMultiTransformBlendHelper& BlendHelperInLocalSpace) const override;
+	ANIMATIONCORE_API virtual void AccumulateConstraintTransform(const FTransform& TargetTransform, const FTransform& CurrentTransform, const FTransform& CurrentParentTransform, float Weight, FMultiTransformBlendHelper& BlendHelperInLocalSpace) const override;
 	virtual bool DoesAffectRotation() const override { return true; }
 	virtual FString GetDisplayString() const override
 	{
@@ -678,7 +678,7 @@ public:
  * You can have as many of these per node
  */
 USTRUCT()
-struct ANIMATIONCORE_API FConstraintData
+struct FConstraintData
 {
 	GENERATED_USTRUCT_BODY()
 
@@ -728,12 +728,12 @@ struct ANIMATIONCORE_API FConstraintData
 		return Ar;
 	}
 
-	void ApplyInverseOffset(const FTransform& InTarget, FTransform& OutSource, const FTransform& InBaseTransform) const;
-	void SaveInverseOffset(const FTransform& Source, const FTransform& Target, const FTransform& InBaseTransform);
+	ANIMATIONCORE_API void ApplyInverseOffset(const FTransform& InTarget, FTransform& OutSource, const FTransform& InBaseTransform) const;
+	ANIMATIONCORE_API void SaveInverseOffset(const FTransform& Source, const FTransform& Target, const FTransform& InBaseTransform);
 	void ResetOffset()
 	{
 		Offset = FTransform::Identity;
 	}
 
-	void ApplyConstraintTransform(const FTransform& TargetTransform, const FTransform& InCurrentTransform, const FTransform& CurrentParentTransform, FMultiTransformBlendHelper& BlendHelperInLocalSpace) const;
-};
\ No newline at end of file
+	ANIMATIONCORE_API void ApplyConstraintTransform(const FTransform& TargetTransform, const FTransform& InCurrentTransform, const FTransform& CurrentParentTransform, FMultiTransformBlendHelper& BlendHelperInLocalSpace) const;
+};
diff --git a/Engine/Source/Runtime/AnimationCore/Public/EulerTransform.h b/Engine/Source/Runtime/AnimationCore/Public/EulerTransform.h
--- a/Engine/Source/Runtime/AnimationCore/Public/EulerTransform.h
+++ b/Engine/Source/Runtime/AnimationCore/Public/EulerTransform.h
@@ -28,7 +28,7 @@ enum class EEulerRotationOrder : uint8
 };
 
 USTRUCT(BlueprintType)
-struct ANIMATIONCORE_API FEulerTransform
+struct FEulerTransform
 {
 	GENERATED_BODY()
 
@@ -37,7 +37,7 @@ struct ANIMATIONCORE_API FEulerTransform
 	/**
 	 * The identity transformation (Rotation = FRotator::ZeroRotator, Translation = FVector::ZeroVector, Scale = (1,1,1)).
 	 */
-	static const FEulerTransform Identity;
+	static ANIMATIONCORE_API const FEulerTransform Identity;
 
 	FORCEINLINE FEulerTransform()
 		: Location(ForceInitToZero)
@@ -115,5 +115,5 @@ struct ANIMATIONCORE_API FEulerTransform
 
 template<> struct TBaseStructure<FEulerTransform>
 {
-	ANIMATIONCORE_API static UScriptStruct* Get() { return FEulerTransform::StaticStruct(); }
+	static UScriptStruct* Get() { return FEulerTransform::StaticStruct(); }
 };
diff --git a/Engine/Source/Runtime/AnimationCore/Public/NodeChain.h b/Engine/Source/Runtime/AnimationCore/Public/NodeChain.h
--- a/Engine/Source/Runtime/AnimationCore/Public/NodeChain.h
+++ b/Engine/Source/Runtime/AnimationCore/Public/NodeChain.h
@@ -7,10 +7,10 @@
 
 /** A chain of nodes in a hierarchy */
 USTRUCT()
-struct ANIMATIONCORE_API FNodeChain
+struct FNodeChain
 {
 	GENERATED_BODY()
 
 	UPROPERTY()
 	TArray<FName> Nodes;
-};
\ No newline at end of file
+};
diff --git a/Engine/Source/Runtime/AnimationCore/Public/TransformNoScale.h b/Engine/Source/Runtime/AnimationCore/Public/TransformNoScale.h
--- a/Engine/Source/Runtime/AnimationCore/Public/TransformNoScale.h
+++ b/Engine/Source/Runtime/AnimationCore/Public/TransformNoScale.h
@@ -13,14 +13,14 @@
 #include "TransformNoScale.generated.h"
 
 USTRUCT(BlueprintType)
-struct ANIMATIONCORE_API FTransformNoScale
+struct FTransformNoScale
 {
 	GENERATED_BODY()
 
 	/**
 	 * The identity transformation (Rotation = FRotator::ZeroRotator, Translation = FVector::ZeroVector, Scale = (1,1,1)).
 	 */
-	static const FTransformNoScale Identity;
+	static ANIMATIONCORE_API const FTransformNoScale Identity;
 
 	FORCEINLINE_DEBUGGABLE FTransformNoScale()
 		: Location(ForceInitToZero)
@@ -71,4 +71,4 @@ struct ANIMATIONCORE_API FTransformNoScale
 		Location = InTransform.GetLocation();
 		Rotation = InTransform.GetRotation();
 	}
-};
\ No newline at end of file
+};
diff --git a/Engine/Source/Runtime/AppFramework/Public/Widgets/Colors/SColorGradingPicker.h b/Engine/Source/Runtime/AppFramework/Public/Widgets/Colors/SColorGradingPicker.h
--- a/Engine/Source/Runtime/AppFramework/Public/Widgets/Colors/SColorGradingPicker.h
+++ b/Engine/Source/Runtime/AppFramework/Public/Widgets/Colors/SColorGradingPicker.h
@@ -33,7 +33,7 @@ enum class EColorGradingModes
  * use the functions OpenColorGradingWheel and DestroyColorGradingWheel, since they hold a static
  * instance of the color picker.
  */
-class APPFRAMEWORK_API SColorGradingPicker
+class SColorGradingPicker
 	: public SCompoundWidget
 {
 public:
@@ -93,7 +93,7 @@ public:
 	SLATE_END_ARGS()
 
 	/**	Destructor. */
-	~SColorGradingPicker();
+	APPFRAMEWORK_API ~SColorGradingPicker();
 
 public:
 
@@ -102,38 +102,38 @@ public:
 	 *
 	 * @param InArgs Declaration from which to construct the widget.
 	 */
-	void Construct(const FArguments& InArgs );
+	APPFRAMEWORK_API void Construct(const FArguments& InArgs );
 
 	FOnNumericEntryBoxDynamicSliderMinMaxValueChanged& GetOnNumericEntryBoxDynamicSliderMaxValueChangedDelegate() { return OnNumericEntryBoxDynamicSliderMaxValueChanged; }
 	FOnNumericEntryBoxDynamicSliderMinMaxValueChanged& GetOnNumericEntryBoxDynamicSliderMinValueChangedDelegate() { return OnNumericEntryBoxDynamicSliderMinValueChanged; }
 
 	/** Callback when the max/min spinner value are changed (only apply if SupportDynamicSliderMaxValue or SupportDynamicSliderMinValue are true) */
-	void OnDynamicSliderMaxValueChanged(float NewMaxSliderValue, TWeakPtr<SWidget> InValueChangedSourceWidget, bool IsOriginator, bool UpdateOnlyIfHigher);
-	void OnDynamicSliderMinValueChanged(float NewMinSliderValue, TWeakPtr<SWidget> InValueChangedSourceWidget, bool IsOriginator, bool UpdateOnlyIfLower);
+	APPFRAMEWORK_API void OnDynamicSliderMaxValueChanged(float NewMaxSliderValue, TWeakPtr<SWidget> InValueChangedSourceWidget, bool IsOriginator, bool UpdateOnlyIfHigher);
+	APPFRAMEWORK_API void OnDynamicSliderMinValueChanged(float NewMinSliderValue, TWeakPtr<SWidget> InValueChangedSourceWidget, bool IsOriginator, bool UpdateOnlyIfLower);
 
 protected:
 
-	void TransformLinearColorRangeToColorGradingRange(FVector4 &VectorValue) const;
-	void TransformColorGradingRangeToLinearColorRange(FVector4 &VectorValue) const;
-	void TransformColorGradingRangeToLinearColorRange(float &FloatValue);
+	APPFRAMEWORK_API void TransformLinearColorRangeToColorGradingRange(FVector4 &VectorValue) const;
+	APPFRAMEWORK_API void TransformColorGradingRangeToLinearColorRange(FVector4 &VectorValue) const;
+	APPFRAMEWORK_API void TransformColorGradingRangeToLinearColorRange(float &FloatValue);
 
-	TOptional<float> OnGetMainValue() const;
-	void OnMainValueChanged(float InValue, bool ShouldCommitValueChanges);
-	void OnMainValueCommitted(float InValue, ETextCommit::Type CommitType);
+	APPFRAMEWORK_API TOptional<float> OnGetMainValue() const;
+	APPFRAMEWORK_API void OnMainValueChanged(float InValue, bool ShouldCommitValueChanges);
+	APPFRAMEWORK_API void OnMainValueCommitted(float InValue, ETextCommit::Type CommitType);
 
-	FLinearColor GetCurrentLinearColor();
+	APPFRAMEWORK_API FLinearColor GetCurrentLinearColor();
 
-	bool IsEntryBoxEnabled() const;
+	APPFRAMEWORK_API bool IsEntryBoxEnabled() const;
 
 	// Callback for value changes in the color spectrum picker.
-	void HandleCurrentColorValueChanged(const FLinearColor& NewValue, bool ShouldCommitValueChanges);
+	APPFRAMEWORK_API void HandleCurrentColorValueChanged(const FLinearColor& NewValue, bool ShouldCommitValueChanges);
 
-	void HandleColorWheelMouseCaptureBegin(const FLinearColor& InValue);
-	void HandleColorWheelMouseCaptureEnd(const FLinearColor& InValue);
+	APPFRAMEWORK_API void HandleColorWheelMouseCaptureBegin(const FLinearColor& InValue);
+	APPFRAMEWORK_API void HandleColorWheelMouseCaptureEnd(const FLinearColor& InValue);
 
-	void OnBeginSliderMovement();
-	void OnEndSliderMovement(float NewValue);
-	void AdjustRatioValue(FVector4 &NewValue);
+	APPFRAMEWORK_API void OnBeginSliderMovement();
+	APPFRAMEWORK_API void OnEndSliderMovement(float NewValue);
+	APPFRAMEWORK_API void AdjustRatioValue(FVector4 &NewValue);
 
 	bool bIsMouseDragging;
 	FVector4 StartDragRatio;
diff --git a/Engine/Source/Runtime/AppFramework/Public/Widgets/Colors/SColorPicker.h b/Engine/Source/Runtime/AppFramework/Public/Widgets/Colors/SColorPicker.h
--- a/Engine/Source/Runtime/AppFramework/Public/Widgets/Colors/SColorPicker.h
+++ b/Engine/Source/Runtime/AppFramework/Public/Widgets/Colors/SColorPicker.h
@@ -69,7 +69,7 @@ struct FColorChannels
  * use the functions OpenColorPicker and DestroyColorPicker, since they hold a static
  * instance of the color picker.
  */
-class APPFRAMEWORK_API SColorPicker
+class SColorPicker
 	: public SCompoundWidget
 {
 public:
@@ -165,10 +165,10 @@ public:
 	SLATE_END_ARGS()
 	
 	/** A default window size for the color picker which looks nice */
-	static const FVector2D DEFAULT_WINDOW_SIZE;
+	static APPFRAMEWORK_API const FVector2D DEFAULT_WINDOW_SIZE;
 
 	/**	Destructor. */
-	~SColorPicker();
+	APPFRAMEWORK_API ~SColorPicker();
 
 public:
 
@@ -177,7 +177,7 @@ public:
 	 *
 	 * @param InArgs Declaration from which to construct the widget.
 	 */
-	void Construct(const FArguments& InArgs);
+	APPFRAMEWORK_API void Construct(const FArguments& InArgs);
 
 	/** Gets the (optionally set) owning details view of the current color picker */
 	TSharedPtr<SWidget> GetOptionalOwningDetailsView()
@@ -194,21 +194,21 @@ public:
 
 	/** Delegate to override color picker creation behavior */
 	DECLARE_DELEGATE_OneParam(FOnColorPickerCreationOverride, const TSharedRef<SColorPicker>&);
-	static FOnColorPickerCreationOverride OnColorPickerNonModalCreateOverride;
+	static APPFRAMEWORK_API FOnColorPickerCreationOverride OnColorPickerNonModalCreateOverride;
 
 	/** Delegate to override color picker destruction behavior */
 	DECLARE_DELEGATE(FOnColorPickerDestructionOverride);
-	static FOnColorPickerDestructionOverride OnColorPickerDestroyOverride;
+	static APPFRAMEWORK_API FOnColorPickerDestructionOverride OnColorPickerDestroyOverride;
 
 protected:
 
 	/** Backup all the colors that are being modified */
-	void BackupColors();
+	APPFRAMEWORK_API void BackupColors();
 
-	bool ApplyNewTargetColor(bool bForceUpdate = false);
+	APPFRAMEWORK_API bool ApplyNewTargetColor(bool bForceUpdate = false);
 
-	void GenerateDefaultColorPickerContent(bool bAdvancedSectionExpanded);
-	void GenerateInlineColorPickerContent();
+	APPFRAMEWORK_API void GenerateDefaultColorPickerContent(bool bAdvancedSectionExpanded);
+	APPFRAMEWORK_API void GenerateInlineColorPickerContent();
 
 	FLinearColor GetCurrentColor() const
 	{
@@ -216,22 +216,22 @@ protected:
 	}
 
 	/** Calls the user defined delegate for when the color changes are discarded */
-	void DiscardColor();
+	APPFRAMEWORK_API void DiscardColor();
 
 	/** Sets new color in ether RGB or HSV */
-	bool SetNewTargetColorRGB(const FLinearColor& NewValue, bool bForceUpdate = false);
-	bool SetNewTargetColorHSV(const FLinearColor& NewValue, bool bForceUpdate = false);
+	APPFRAMEWORK_API bool SetNewTargetColorRGB(const FLinearColor& NewValue, bool bForceUpdate = false);
+	APPFRAMEWORK_API bool SetNewTargetColorHSV(const FLinearColor& NewValue, bool bForceUpdate = false);
 
-	void UpdateColorPick();
-	void UpdateColorPickMouseUp();
+	APPFRAMEWORK_API void UpdateColorPick();
+	APPFRAMEWORK_API void UpdateColorPickMouseUp();
 
-	void BeginAnimation(FLinearColor Start, FLinearColor End);
+	APPFRAMEWORK_API void BeginAnimation(FLinearColor Start, FLinearColor End);
 
-	void HideSmallTrash();
-	void ShowSmallTrash();
+	APPFRAMEWORK_API void HideSmallTrash();
+	APPFRAMEWORK_API void ShowSmallTrash();
 
 	/** Cycles the color picker's mode. */
-	void CycleMode();
+	APPFRAMEWORK_API void CycleMode();
 
 	/**
 	 * Creates a color slider widget for the specified channel.
@@ -239,7 +239,7 @@ protected:
 	 * @param Channel The color channel to create the widget for.
 	 * @return The new slider.
 	 */
-	TSharedRef<SWidget> MakeColorSlider(EColorPickerChannels Channel) const;
+	APPFRAMEWORK_API TSharedRef<SWidget> MakeColorSlider(EColorPickerChannels Channel) const;
 
 	/**
 	 * Creates a color spin box widget for the specified channel.
@@ -247,124 +247,124 @@ protected:
 	 * @param Channel The color channel to create the widget for.
 	 * @return The new spin box.
 	 */
-	TSharedRef<SWidget> MakeColorSpinBox(EColorPickerChannels Channel) const;
+	APPFRAMEWORK_API TSharedRef<SWidget> MakeColorSpinBox(EColorPickerChannels Channel) const;
 
 	/**
 	 * Creates the color preview box widget.
 	 *
 	 * @return The new color preview box.
 	 */
-	TSharedRef<SWidget> MakeColorPreviewBox() const;
+	APPFRAMEWORK_API TSharedRef<SWidget> MakeColorPreviewBox() const;
 
 private:
 
 	// Callback for the active timer to animate the color post-construct
-	EActiveTimerReturnType AnimatePostConstruct(double InCurrentTime, float InDeltaTime);
+	APPFRAMEWORK_API EActiveTimerReturnType AnimatePostConstruct(double InCurrentTime, float InDeltaTime);
 
 	// Callback for getting the end color of a color spin box gradient.
-	FLinearColor GetGradientEndColor(EColorPickerChannels Channel) const;
+	APPFRAMEWORK_API FLinearColor GetGradientEndColor(EColorPickerChannels Channel) const;
 
 	// Callback for getting the start color of a color spin box gradient.
-	FLinearColor GetGradientStartColor(EColorPickerChannels Channel) const;
+	APPFRAMEWORK_API FLinearColor GetGradientStartColor(EColorPickerChannels Channel) const;
 
 	// Callback for handling expansion of the 'Advanced' area.
-	void HandleAdvancedAreaExpansionChanged(bool Expanded);
+	APPFRAMEWORK_API void HandleAdvancedAreaExpansionChanged(bool Expanded);
 
 	// Callback for getting the visibility of the alpha channel portion in color blocks.
-	EVisibility HandleAlphaColorBlockVisibility() const;
+	APPFRAMEWORK_API EVisibility HandleAlphaColorBlockVisibility() const;
 
 	// Callback for clicking the Cancel button.
-	FReply HandleCancelButtonClicked();
+	APPFRAMEWORK_API FReply HandleCancelButtonClicked();
 
 	// Callback for pressing a mouse button in the color area.
-	FReply HandleColorAreaMouseDown(const FGeometry& MyGeometry, const FPointerEvent& MouseEvent);
+	APPFRAMEWORK_API FReply HandleColorAreaMouseDown(const FGeometry& MyGeometry, const FPointerEvent& MouseEvent);
 
 	// Callback for clicking the color picker mode button.
-	FReply HandleColorPickerModeButtonClicked();
+	APPFRAMEWORK_API FReply HandleColorPickerModeButtonClicked();
 
 	// Callback for getting the visibility of the given color picker mode.
-	EVisibility HandleColorPickerModeVisibility(EColorPickerModes Mode) const;
+	APPFRAMEWORK_API EVisibility HandleColorPickerModeVisibility(EColorPickerModes Mode) const;
 
 	// Callback for getting the end color of a color slider.
-	FLinearColor HandleColorSliderEndColor(EColorPickerChannels Channel) const;
+	APPFRAMEWORK_API FLinearColor HandleColorSliderEndColor(EColorPickerChannels Channel) const;
 
 	// Callback for getting the start color of a color slider.
-	FLinearColor HandleColorSliderStartColor(EColorPickerChannels Channel) const;
+	APPFRAMEWORK_API FLinearColor HandleColorSliderStartColor(EColorPickerChannels Channel) const;
 
 	// Callback for value changes in the color spectrum picker.
-	void HandleColorSpectrumValueChanged(FLinearColor NewValue);
+	APPFRAMEWORK_API void HandleColorSpectrumValueChanged(FLinearColor NewValue);
 
 	// Callback for getting the value of a color spin box.
-	float HandleColorSpinBoxValue(EColorPickerChannels Channel) const;
+	APPFRAMEWORK_API float HandleColorSpinBoxValue(EColorPickerChannels Channel) const;
 
 	// Callback for value changes in a color spin box.
-	void HandleColorSpinBoxValueChanged(float NewValue, EColorPickerChannels Channel);
+	APPFRAMEWORK_API void HandleColorSpinBoxValueChanged(float NewValue, EColorPickerChannels Channel);
 
 	// Callback for completed eye dropper interactions.
-	void HandleEyeDropperButtonComplete(bool bCancelled);
+	APPFRAMEWORK_API void HandleEyeDropperButtonComplete(bool bCancelled);
 
 	// Callback for getting the text in the hex linear box.
-	FText HandleHexLinearBoxText() const;
+	APPFRAMEWORK_API FText HandleHexLinearBoxText() const;
 
 	// Callback for getting the text in the hex sRGB box.
-	FText HandleHexSRGBBoxText() const;
+	APPFRAMEWORK_API FText HandleHexSRGBBoxText() const;
 
 	// Callback for committed text in the hex input box (sRGB gamma).
-	void HandleHexSRGBInputTextCommitted(const FText& Text, ETextCommit::Type CommitType);
+	APPFRAMEWORK_API void HandleHexSRGBInputTextCommitted(const FText& Text, ETextCommit::Type CommitType);
 
 	// Callback for committed text in the hex input box (linear gamma).
-	void HandleHexLinearInputTextCommitted(const FText& Text, ETextCommit::Type CommitType);
+	APPFRAMEWORK_API void HandleHexLinearInputTextCommitted(const FText& Text, ETextCommit::Type CommitType);
 
 	// Callback for changing the HSV value of the current color.
-	void HandleHSVColorChanged(FLinearColor NewValue);
+	APPFRAMEWORK_API void HandleHSVColorChanged(FLinearColor NewValue);
 
 	// Callback for when interactive user input begins.
-	void HandleInteractiveChangeBegin();
+	APPFRAMEWORK_API void HandleInteractiveChangeBegin();
 
 	// Callback for when interactive user input ends.
-	void HandleInteractiveChangeEnd();
+	APPFRAMEWORK_API void HandleInteractiveChangeEnd();
 
 	// Callback for when interactive user input ends.
-	void HandleInteractiveChangeEnd(float NewValue);
+	APPFRAMEWORK_API void HandleInteractiveChangeEnd(float NewValue);
 
 	// Callback for clicking the new color preview block.
-	FReply HandleNewColorBlockMouseButtonDown(const FGeometry& MyGeometry, const FPointerEvent& MouseEvent, bool bCheckAlpha);
+	APPFRAMEWORK_API FReply HandleNewColorBlockMouseButtonDown(const FGeometry& MyGeometry, const FPointerEvent& MouseEvent, bool bCheckAlpha);
 
 	// Callback for clicking the OK button.
-	FReply HandleOkButtonClicked();
+	APPFRAMEWORK_API FReply HandleOkButtonClicked();
 
 	// Callback for clicking the old color preview block.
-	FReply HandleOldColorBlockMouseButtonDown(const FGeometry& MyGeometry, const FPointerEvent& MouseEvent, bool bCheckAlpha);
+	APPFRAMEWORK_API FReply HandleOldColorBlockMouseButtonDown(const FGeometry& MyGeometry, const FPointerEvent& MouseEvent, bool bCheckAlpha);
 
 	// Callback for checking whether sRGB colors should be rendered.
-	bool HandleColorPickerUseSRGB() const;
+	APPFRAMEWORK_API bool HandleColorPickerUseSRGB() const;
 
 	// Callback for when the parent window has been closed.
-	void HandleParentWindowClosed(const TSharedRef<SWindow>& Window);
+	APPFRAMEWORK_API void HandleParentWindowClosed(const TSharedRef<SWindow>& Window);
 
 	// Callback for changing the RGB value of the current color.
-	void HandleRGBColorChanged(FLinearColor NewValue);
+	APPFRAMEWORK_API void HandleRGBColorChanged(FLinearColor NewValue);
 
 	// Callback for changing the checked state of the sRGB check box.
-	void HandleSRGBCheckBoxCheckStateChanged(ECheckBoxState InIsChecked);
+	APPFRAMEWORK_API void HandleSRGBCheckBoxCheckStateChanged(ECheckBoxState InIsChecked);
 
 	// Callback for determining whether the sRGB check box should be checked.
-	ECheckBoxState HandleSRGBCheckBoxIsChecked() const;
+	APPFRAMEWORK_API ECheckBoxState HandleSRGBCheckBoxIsChecked() const;
 
 	// Callback for selecting a color in the color theme bar.
-	void HandleThemeBarColorSelected(FLinearColor NewValue);
+	APPFRAMEWORK_API void HandleThemeBarColorSelected(FLinearColor NewValue);
 
 	// Callback for getting the theme bar's color theme.
-	TSharedPtr<class FColorTheme> HandleThemeBarColorTheme() const;
+	APPFRAMEWORK_API TSharedPtr<class FColorTheme> HandleThemeBarColorTheme() const;
 
 	// Callback for getting the visibility of the theme bar hint text.
-	EVisibility HandleThemeBarHintVisibility() const;
+	APPFRAMEWORK_API EVisibility HandleThemeBarHintVisibility() const;
 
 	// Callback for determining whether the theme bar should display the alpha channel.
-	bool HandleThemeBarUseAlpha() const;
+	APPFRAMEWORK_API bool HandleThemeBarUseAlpha() const;
 
 	// Callback for theme viewer changes.
-	void HandleThemesViewerThemeChanged();
+	APPFRAMEWORK_API void HandleThemesViewerThemeChanged();
 
 private:
 	
@@ -393,7 +393,7 @@ private:
 	float CurrentTime;
 
 	/** The max time allowed for updating before we shut off auto-updating */
-	static const double MAX_ALLOWED_UPDATE_TIME;
+	static APPFRAMEWORK_API const double MAX_ALLOWED_UPDATE_TIME;
 
 	/** If true, then the performance is too bad to have auto-updating */
 	bool bPerfIsTooSlowToUpdate;
@@ -469,7 +469,7 @@ private:
 private:
 
 	/** A static pointer to the global color themes viewer */
-	static TWeakPtr<SColorThemesViewer> ColorThemesViewer;
+	static APPFRAMEWORK_API TWeakPtr<SColorThemesViewer> ColorThemesViewer;
 };
 
 
diff --git a/Engine/Source/Runtime/AppFramework/Public/Widgets/Workflow/SWizard.h b/Engine/Source/Runtime/AppFramework/Public/Widgets/Workflow/SWizard.h
--- a/Engine/Source/Runtime/AppFramework/Public/Widgets/Workflow/SWizard.h
+++ b/Engine/Source/Runtime/AppFramework/Public/Widgets/Workflow/SWizard.h
@@ -30,7 +30,7 @@ DECLARE_DELEGATE_RetVal_OneParam(int32, FOnGetPageIndex, int32)
 /**
  * Implements a wizard widget.
  */
-class APPFRAMEWORK_API SWizard
+class SWizard
 	: public SCompoundWidget
 {
 public:
@@ -38,7 +38,7 @@ public:
 	/**
 	 * Implements a wizard page.
 	 */
-	class APPFRAMEWORK_API FWizardPage
+	class FWizardPage
 	{
 	public:
 
@@ -263,27 +263,27 @@ public:
 	 * @param PageIndex The index of the page to check.
 	 * @return true if the page can be shown, false otherwise.
 	 */
-	bool CanShowPage( int32 PageIndex ) const;
+	APPFRAMEWORK_API bool CanShowPage( int32 PageIndex ) const;
 
 	/**
 	 * Construct this widget.
 	 *
 	 * @param InArgs The declaration data for this widget.
 	 */
-	void Construct( const FArguments& InArgs );
+	APPFRAMEWORK_API void Construct( const FArguments& InArgs );
 
 	/**
 	 * Gets the number of pages that this wizard contains.
 	 *
 	 * @return The number of pages.
 	 */
-	int32 GetNumPages() const;
+	APPFRAMEWORK_API int32 GetNumPages() const;
 
 	/**
 	 * Get the index of the current page.
 	 * @return The index of the current page.
 	 */
-	int32 GetCurrentPageIndex() const;
+	APPFRAMEWORK_API int32 GetCurrentPageIndex() const;
 
 	/**
 	 * Gets the index of the specified wizard page widget.
@@ -291,23 +291,23 @@ public:
 	 * @param PageWidget The page widget to get the index for.
 	 * @return The index of the page, or INDEX_NONE if not found.
 	 */
-	int32 GetPageIndex( const TSharedPtr<SWidget>& PageWidget ) const;
+	APPFRAMEWORK_API int32 GetPageIndex( const TSharedPtr<SWidget>& PageWidget ) const;
 
 	/**
 	 * Attempts to show the page with the specified index.
 	 *
 	 * @param PageIndex The index of the page to show.
 	 */
-	void ShowPage( int32 PageIndex );
+	APPFRAMEWORK_API void ShowPage( int32 PageIndex );
 
 	/** Shows the page with the specified index, but preserves page history. */
-	void AdvanceToPage( int32 PageIndex );
+	APPFRAMEWORK_API void AdvanceToPage( int32 PageIndex );
 
 public:
 
 	// SCompoundWidget interface
 
-	virtual FVector2D ComputeDesiredSize(float) const override;
+	APPFRAMEWORK_API virtual FVector2D ComputeDesiredSize(float) const override;
 
 public:
 
@@ -326,49 +326,49 @@ public:
 private:
 
 	// Callback for clicking the 'Cancel' button.
-	FReply HandleCancelButtonClicked();
+	APPFRAMEWORK_API FReply HandleCancelButtonClicked();
 
 	// Callback for clicking the 'Launch' button.
-	FReply HandleFinishButtonClicked();
+	APPFRAMEWORK_API FReply HandleFinishButtonClicked();
 
 	// Callback for clicking the 'Next' button.
-	FReply HandleNextButtonClicked();
+	APPFRAMEWORK_API FReply HandleNextButtonClicked();
 
 	// Callback for getting the enabled state of the 'Next' button.
-	bool HandleNextButtonIsEnabled() const;
+	APPFRAMEWORK_API bool HandleNextButtonIsEnabled() const;
 
 	// Callback for getting the visibility of the 'Next' button.
-	EVisibility HandleNextButtonVisibility() const;
+	APPFRAMEWORK_API EVisibility HandleNextButtonVisibility() const;
 
 	// Callback for getting the checked state of a page button.
-	void HandlePageButtonCheckStateChanged( ECheckBoxState NewState, int32 PageIndex );
+	APPFRAMEWORK_API void HandlePageButtonCheckStateChanged( ECheckBoxState NewState, int32 PageIndex );
 
 	// Callback for clicking a page button.
-	ECheckBoxState HandlePageButtonIsChecked( int32 PageIndex ) const;
+	APPFRAMEWORK_API ECheckBoxState HandlePageButtonIsChecked( int32 PageIndex ) const;
 
 	// Callback for getting the enabled state of a page button.
-	bool HandlePageButtonIsEnabled( int32 PageIndex ) const;
+	APPFRAMEWORK_API bool HandlePageButtonIsEnabled( int32 PageIndex ) const;
 
 	// Callback for clicking the 'Previous' button.
-	FReply HandlePrevButtonClicked();
+	APPFRAMEWORK_API FReply HandlePrevButtonClicked();
 
 	// Callback for getting the enabled state of the 'Previous' button.
-	bool HandlePrevButtonIsEnabled() const;
+	APPFRAMEWORK_API bool HandlePrevButtonIsEnabled() const;
 
 	// Callback for getting the visibility of the 'Previous' button.
-	EVisibility HandlePrevButtonVisibility() const;
+	APPFRAMEWORK_API EVisibility HandlePrevButtonVisibility() const;
 
 	// Callback for when a breadcrumb is clicked.
-	void HandleBreadcrumbClicked(const int32& PageIndex);
+	APPFRAMEWORK_API void HandleBreadcrumbClicked(const int32& PageIndex);
 
 	// Get the next page index to display.
-	int32 GetNextPageIndex() const;
+	APPFRAMEWORK_API int32 GetNextPageIndex() const;
 
 	// Get the previous page index to display.
-	int32 GetPrevPageIndex() const;
+	APPFRAMEWORK_API int32 GetPrevPageIndex() const;
 
 	// Get the page title to display for the current page.
-	FText HandleGetPageTitle() const;
+	APPFRAMEWORK_API FText HandleGetPageTitle() const;
 
 private:
 
diff --git a/Engine/Source/Runtime/ApplicationCore/Public/Android/AndroidPlatformApplicationMisc.h b/Engine/Source/Runtime/ApplicationCore/Public/Android/AndroidPlatformApplicationMisc.h
--- a/Engine/Source/Runtime/ApplicationCore/Public/Android/AndroidPlatformApplicationMisc.h
+++ b/Engine/Source/Runtime/ApplicationCore/Public/Android/AndroidPlatformApplicationMisc.h
@@ -4,25 +4,25 @@
 
 #include "GenericPlatform/GenericPlatformApplicationMisc.h"
 
-struct APPLICATIONCORE_API FAndroidApplicationMisc : public FGenericPlatformApplicationMisc
+struct FAndroidApplicationMisc : public FGenericPlatformApplicationMisc
 {
-	static void LoadPreInitModules();
-	static class FFeedbackContext* GetFeedbackContext();
-	static class FOutputDeviceError* GetErrorOutputDevice();
-	static class GenericApplication* CreateApplication();
-	static void RequestMinimize();
-	static bool IsScreensaverEnabled();
-	static bool ControlScreensaver(EScreenSaverAction Action);
-	static void SetGamepadsAllowed(bool bAllowed);
-	static void SetGamepadsBlockDeviceFeedback(bool bBlock);
-	static void ResetGamepadAssignments();
-	static void ResetGamepadAssignmentToController(int32 ControllerId);
-	static bool IsControllerAssignedToGamepad(int32 ControllerId);
-	static FString GetGamepadControllerName(int32 ControllerId);
-	static void ClipboardCopy(const TCHAR* Str);
-	static void ClipboardPaste(class FString& Dest);
-	static EScreenPhysicalAccuracy ComputePhysicalScreenDensity(int32& OutScreenDensity);
-	static void EnableMotionData(bool bEnable);
+	static APPLICATIONCORE_API void LoadPreInitModules();
+	static APPLICATIONCORE_API class FFeedbackContext* GetFeedbackContext();
+	static APPLICATIONCORE_API class FOutputDeviceError* GetErrorOutputDevice();
+	static APPLICATIONCORE_API class GenericApplication* CreateApplication();
+	static APPLICATIONCORE_API void RequestMinimize();
+	static APPLICATIONCORE_API bool IsScreensaverEnabled();
+	static APPLICATIONCORE_API bool ControlScreensaver(EScreenSaverAction Action);
+	static APPLICATIONCORE_API void SetGamepadsAllowed(bool bAllowed);
+	static APPLICATIONCORE_API void SetGamepadsBlockDeviceFeedback(bool bBlock);
+	static APPLICATIONCORE_API void ResetGamepadAssignments();
+	static APPLICATIONCORE_API void ResetGamepadAssignmentToController(int32 ControllerId);
+	static APPLICATIONCORE_API bool IsControllerAssignedToGamepad(int32 ControllerId);
+	static APPLICATIONCORE_API FString GetGamepadControllerName(int32 ControllerId);
+	static APPLICATIONCORE_API void ClipboardCopy(const TCHAR* Str);
+	static APPLICATIONCORE_API void ClipboardPaste(class FString& Dest);
+	static APPLICATIONCORE_API EScreenPhysicalAccuracy ComputePhysicalScreenDensity(int32& OutScreenDensity);
+	static APPLICATIONCORE_API void EnableMotionData(bool bEnable);
 };
 
 typedef FAndroidApplicationMisc FPlatformApplicationMisc;
diff --git a/Engine/Source/Runtime/ApplicationCore/Public/Android/AndroidPlatformSplash.h b/Engine/Source/Runtime/ApplicationCore/Public/Android/AndroidPlatformSplash.h
--- a/Engine/Source/Runtime/ApplicationCore/Public/Android/AndroidPlatformSplash.h
+++ b/Engine/Source/Runtime/ApplicationCore/Public/Android/AndroidPlatformSplash.h
@@ -10,7 +10,7 @@
 /**
  * Android splash implementation
  */
-struct APPLICATIONCORE_API FAndroidSplash : public FGenericPlatformSplash
+struct FAndroidSplash : public FGenericPlatformSplash
 {
 	// default implementation for now
 };
diff --git a/Engine/Source/Runtime/ApplicationCore/Public/Android/AndroidWindow.h b/Engine/Source/Runtime/ApplicationCore/Public/Android/AndroidWindow.h
--- a/Engine/Source/Runtime/ApplicationCore/Public/Android/AndroidWindow.h
+++ b/Engine/Source/Runtime/ApplicationCore/Public/Android/AndroidWindow.h
@@ -15,59 +15,59 @@
  * Native windows provide platform-specific backing for and are always owned by an SWindow.
  */
  
-class APPLICATIONCORE_API FAndroidWindow : public FGenericWindow
+class FAndroidWindow : public FGenericWindow
 {
 public:
-	~FAndroidWindow();
+	APPLICATIONCORE_API ~FAndroidWindow();
 
 	/** Create a new FAndroidWindow.
 	 *
 	 * @param OwnerWindow		The SlateWindow for which we are crating a backing AndroidWindow
 	 * @param InParent			Parent iOS window; usually NULL.
 	 */
-	static TSharedRef<FAndroidWindow> Make();
+	static APPLICATIONCORE_API TSharedRef<FAndroidWindow> Make();
 
 	
 	virtual void* GetOSWindowHandle() const override { return nullptr; }
 
-	void Initialize( class FAndroidApplication* const Application, const TSharedRef< FGenericWindowDefinition >& InDefinition, const TSharedPtr< FAndroidWindow >& InParent, const bool bShowImmediately );
+	APPLICATIONCORE_API void Initialize( class FAndroidApplication* const Application, const TSharedRef< FGenericWindowDefinition >& InDefinition, const TSharedPtr< FAndroidWindow >& InParent, const bool bShowImmediately );
 
 	/** Returns the rectangle of the screen the window is associated with */
-	virtual bool GetFullScreenInfo( int32& X, int32& Y, int32& Width, int32& Height ) const override;
+	APPLICATIONCORE_API virtual bool GetFullScreenInfo( int32& X, int32& Y, int32& Width, int32& Height ) const override;
 
-	virtual void SetOSWindowHandle(void*);
+	APPLICATIONCORE_API virtual void SetOSWindowHandle(void*);
 
-	static FPlatformRect GetScreenRect(bool bUseEventThreadWindow = false);
-	static void InvalidateCachedScreenRect();
+	static APPLICATIONCORE_API FPlatformRect GetScreenRect(bool bUseEventThreadWindow = false);
+	static APPLICATIONCORE_API void InvalidateCachedScreenRect();
 
 	// When bUseEventThreadWindow == false this uses dimensions cached when the game thread processes android events.
 	// When bUseEventThreadWindow == true this uses dimensions directly from the android event thread, unless called from event thread this requires acquiring GAndroidWindowLock to use.
-	static void CalculateSurfaceSize(int32_t& SurfaceWidth, int32_t& SurfaceHeight, bool bUseEventThreadWindow = false);
-	static bool OnWindowOrientationChanged(int Orientation);
+	static APPLICATIONCORE_API void CalculateSurfaceSize(int32_t& SurfaceWidth, int32_t& SurfaceHeight, bool bUseEventThreadWindow = false);
+	static APPLICATIONCORE_API bool OnWindowOrientationChanged(int Orientation);
 
-	static int32 GetDepthBufferPreference();
+	static APPLICATIONCORE_API int32 GetDepthBufferPreference();
 	
-	static void AcquireWindowRef(ANativeWindow* InWindow);
-	static void ReleaseWindowRef(ANativeWindow* InWindow);
+	static APPLICATIONCORE_API void AcquireWindowRef(ANativeWindow* InWindow);
+	static APPLICATIONCORE_API void ReleaseWindowRef(ANativeWindow* InWindow);
 
 	// This returns the current hardware window as set from the event thread.
-	static void* GetHardwareWindow_EventThread();
-	static void SetHardwareWindow_EventThread(void* InWindow);
+	static APPLICATIONCORE_API void* GetHardwareWindow_EventThread();
+	static APPLICATIONCORE_API void SetHardwareWindow_EventThread(void* InWindow);
 
 	/** Waits on the current thread for a hardware window and returns it. 
 	 *  May return nullptr if the application is shutting down.
 	 */
-	static void* WaitForHardwareWindow();
+	static APPLICATIONCORE_API void* WaitForHardwareWindow();
 
-	static bool IsPortraitOrientation();
-	static FVector4 GetSafezone(bool bPortrait);
-	static bool SafezoneUpdated();
+	static APPLICATIONCORE_API bool IsPortraitOrientation();
+	static APPLICATIONCORE_API FVector4 GetSafezone(bool bPortrait);
+	static APPLICATIONCORE_API bool SafezoneUpdated();
 
 	// called by the Android event thread to initially set the current window dimensions.
-	static void SetWindowDimensions_EventThread(ANativeWindow* DimensionWindow);
+	static APPLICATIONCORE_API void SetWindowDimensions_EventThread(ANativeWindow* DimensionWindow);
 
 	// Called by the event manager to update the cached window dimensions to match the event it is processing.
-	static void EventManagerUpdateWindowDimensions(int32 Width, int32 Height);
+	static APPLICATIONCORE_API void EventManagerUpdateWindowDimensions(int32 Width, int32 Height);
 
 protected:
 	/** @return true if the native window is currently in fullscreen mode, false otherwise */
@@ -77,7 +77,7 @@ private:
 	/**
 	 * Protect the constructor; only TSharedRefs of this class can be made.
 	 */
-	FAndroidWindow();
+	APPLICATIONCORE_API FAndroidWindow();
 
 	FAndroidApplication* OwningApplication;
 
@@ -85,12 +85,12 @@ private:
 	int32 RegionX;
 	int32 RegionY;
 
-	static void* NativeWindow;
+	static APPLICATIONCORE_API void* NativeWindow;
 
 	// Waits for the event thread to report an initial window size.
-	static bool WaitForWindowDimensions();
+	static APPLICATIONCORE_API bool WaitForWindowDimensions();
 
-	static bool bAreCachedNativeDimensionsValid;
-	static int32 CachedNativeWindowWidth;
-	static int32 CachedNativeWindowHeight;
+	static APPLICATIONCORE_API bool bAreCachedNativeDimensionsValid;
+	static APPLICATIONCORE_API int32 CachedNativeWindowWidth;
+	static APPLICATIONCORE_API int32 CachedNativeWindowHeight;
 };
