Add Invariant versions of Justification left/right that ignore the flow direction of the current culture.

#tests Used in widget that wants to use fixed left and right justication
#pf 649de0156b55986b7790f074

diff --git a/Engine/Source/Runtime/AppFramework/Private/Widgets/Testing/STestSuite.cpp b/Engine/Source/Runtime/AppFramework/Private/Widgets/Testing/STestSuite.cpp
--- a/Engine/Source/Runtime/AppFramework/Private/Widgets/Testing/STestSuite.cpp
+++ b/Engine/Source/Runtime/AppFramework/Private/Widgets/Testing/STestSuite.cpp
@@ -2379,6 +2379,14 @@ public:
 		{
 			Justification = ETextJustify::Right;
 		}
+		else if ( *Value == TEXT("Invariant Left") )
+		{
+			Justification = ETextJustify::InvariantLeft;
+		}
+		else if ( *Value == TEXT("Invariant Right") )
+		{
+			Justification = ETextJustify::InvariantRight;
+		}
 	}
 
 	FText JustificationGetSelectedText() const
@@ -2396,6 +2404,14 @@ public:
 		{
 			Text = LOCTEXT("TextJustify::Right", "Right");
 		}
+		else if ( Justification == ETextJustify::InvariantLeft )
+		{
+			Text = LOCTEXT("TextJustify::InvariantLeft", "Invariant Left");
+		}
+		else if ( Justification == ETextJustify::InvariantRight )
+		{
+			Text = LOCTEXT("TextJustify::InvariantRight", "Invariant Right");
+		}
 		return Text;
 	}
 
diff --git a/Engine/Source/Runtime/Slate/Private/Framework/Text/TextLayout.cpp b/Engine/Source/Runtime/Slate/Private/Framework/Text/TextLayout.cpp
--- a/Engine/Source/Runtime/Slate/Private/Framework/Text/TextLayout.cpp
+++ b/Engine/Source/Runtime/Slate/Private/Framework/Text/TextLayout.cpp
@@ -486,6 +486,7 @@ void FTextLayout::JustifyLayout()
 		switch (VisualJustification)
 		{
 		case ETextJustify::Left:
+		case ETextJustify::InvariantLeft:
 			{
 				const float ExtraSpace = LineView.Size.X - LineJustificationWidth;
 				OffsetAdjustment.X = -ExtraSpace;
@@ -500,6 +501,7 @@ void FTextLayout::JustifyLayout()
 			break;
 
 		case ETextJustify::Right:
+		case ETextJustify::InvariantRight:
 			{
 				const float ExtraSpace = LayoutWidthNoMargin - LineJustificationWidth;
 				OffsetAdjustment.X = ExtraSpace;
diff --git a/Engine/Source/Runtime/Slate/Public/Framework/Text/TextLayout.h b/Engine/Source/Runtime/Slate/Public/Framework/Text/TextLayout.h
--- a/Engine/Source/Runtime/Slate/Public/Framework/Text/TextLayout.h
+++ b/Engine/Source/Runtime/Slate/Public/Framework/Text/TextLayout.h
@@ -45,6 +45,16 @@ namespace ETextJustify
 		 * When text is flowing right-to-left, this will align text visually to the left.
 		 */
 		Right,
+
+		/**
+		 * Always justify the text to the left, regardless of the flow direction of the current culture.
+		 */
+		InvariantLeft,
+
+		/**
+		 * Always justify the text to the right, regardless of the flow direction of the current culture.
+		 */
+		InvariantRight,
 	};
 }
 
