Color Picker Current Color Theme textbox shows double text

This was happening because there was a hack in SColorPicker.cpp to show the placeholder text for the SThemeColorBlocksBar, which was broken. It has since been fixed, so I removed the hack to fix the placeholder text. At this point the placeholder text was misaligned, so I fixed the padding for it in SColorThemes.cpp. Befpre and after pic is attached to the Jira.

diff --git a/Engine/Source/Runtime/AppFramework/Private/Widgets/Colors/SColorPicker.cpp b/Engine/Source/Runtime/AppFramework/Private/Widgets/Colors/SColorPicker.cpp
--- a/Engine/Source/Runtime/AppFramework/Private/Widgets/Colors/SColorPicker.cpp
+++ b/Engine/Source/Runtime/AppFramework/Private/Widgets/Colors/SColorPicker.cpp
@@ -200,16 +200,6 @@ void SColorPicker::GenerateDefaultColorPickerContent( bool bAdvancedSectionExpan
 											.UseSRGB(SharedThis(this), &SColorPicker::HandleColorPickerUseSRGB)
 											.OnSelectColor(this, &SColorPicker::HandleThemeBarColorSelected)
 									]
-
-								// hack: need to fix SThemeColorBlocksBar::EmptyText to render properly
-								+ SOverlay::Slot()
-									.HAlign(HAlign_Center)
-									.VAlign(VAlign_Center)
-									[
-										SNew(STextBlock)
-										.Text(LOCTEXT("EmptyBarHint", "Drag & drop colors here to save"))
-										.Visibility(this, &SColorPicker::HandleThemeBarHintVisibility)
-									]
 							]
 
 						+ SHorizontalBox::Slot()
diff --git a/Engine/Source/Runtime/AppFramework/Private/Widgets/Colors/SColorThemes.cpp b/Engine/Source/Runtime/AppFramework/Private/Widgets/Colors/SColorThemes.cpp
--- a/Engine/Source/Runtime/AppFramework/Private/Widgets/Colors/SColorThemes.cpp
+++ b/Engine/Source/Runtime/AppFramework/Private/Widgets/Colors/SColorThemes.cpp
@@ -743,7 +743,7 @@ void SThemeColorBlocksBar::Construct(const FArguments& InArgs)
 	if (!InArgs._EmptyText.IsEmpty())
 	{
 		EmptyHintTextBlock = SNew(SBorder)
-			.Padding(1.0f)
+			.Padding( 8.0f , 4.0f , 0.0f , 0.0f )
 			.Content()
 			[
 				SNew(STextBlock)
