Add optional delegate to handle attempting to move to next wizard page and redirect to another page (e.g. error) on failure.

diff --git a/Engine/Source/Runtime/AppFramework/Private/Widgets/Workflow/SWizard.cpp b/Engine/Source/Runtime/AppFramework/Private/Widgets/Workflow/SWizard.cpp
--- a/Engine/Source/Runtime/AppFramework/Private/Widgets/Workflow/SWizard.cpp
+++ b/Engine/Source/Runtime/AppFramework/Private/Widgets/Workflow/SWizard.cpp
@@ -235,6 +235,7 @@ void SWizard::Construct( const FArguments& InArgs )
 	}
 
 	OnGetNextPageIndex = InArgs._OnGetNextPageIndex;
+	OnTryToMoveToPage = InArgs._OnTryToMoveToPage;
 
 	WidgetSwitcher->SetActiveWidgetIndex(INDEX_NONE);
 	ShowPage(InArgs._InitialPageIndex.Get());
@@ -336,6 +337,12 @@ void SWizard::AdvanceToPage(int32 PageIndex)
 	{
 		BreadcrumbTrail->PushCrumb(Pages[CurrentPage].GetName(), CurrentPage);
 	}
+
+	if (OnTryToMoveToPage.IsBound())
+	{
+		PageIndex = OnTryToMoveToPage.Execute(PageIndex);
+	}
+
 	ShowPage(PageIndex);
 }
 
diff --git a/Engine/Source/Runtime/AppFramework/Public/Widgets/Workflow/SWizard.h b/Engine/Source/Runtime/AppFramework/Public/Widgets/Workflow/SWizard.h
--- a/Engine/Source/Runtime/AppFramework/Public/Widgets/Workflow/SWizard.h
+++ b/Engine/Source/Runtime/AppFramework/Public/Widgets/Workflow/SWizard.h
@@ -224,11 +224,20 @@ public:
 		SLATE_EVENT(FOnClicked, OnFirstPageBackClicked)
 
 		/**
-		 * Exposes a delegate to be invoked when the wizard's 'Next' button is clicked.
+		 * Exposes a delegate to be invoked to control whether the wizard's 'Next' button should be active, and which page should be next if it is clicked.
+		 * Significant work should not be done in this delegate as it will be polled regularly by the button.
 		 * Passes in the current page index and expects the next page index to be returned.
 		 */
 		SLATE_EVENT(FOnGetPageIndex, OnGetNextPageIndex)
 
+		/**
+		 * Exposes a delegate to be invoked when the wizard tries to move to new next page e.g. when 'Next' is pressed.
+		 * Allows the wizard to redirect to another page if processing of the current step fails in some way.
+		 * Passes in the page index the wizard is trying to move to and expects the actual next page index to be returned.
+		 * This process will not be repeated; the replacement page must be able to be moved to without failure.
+		 */
+		SLATE_EVENT(FOnGetPageIndex, OnTryToMoveToPage)
+
 		/** Holds a flag indicating whether the page list should be shown (default = true). */
 		SLATE_ARGUMENT(bool, ShowPageList)
 
@@ -401,6 +410,9 @@ private:
 	// Holds a delegate to be invoked when the 'Previous' button has been clicked on the first page.
 	FOnClicked OnFirstPageBackClicked;
 
-	// Exposes a delegate to be invoked when the wizard's 'Next' button is clicked to fetch the next page index.
+	// Exposes a delegate to be invoked to control whether the 'Next' button is enabled, and where to attempt to move to if it is clicked.
 	FOnGetPageIndex OnGetNextPageIndex;
+
+	// Exposes a delegate to be invoked when the wizard's 'Next' button is clicked to allow redirection to another page if processing fails
+	FOnGetPageIndex OnTryToMoveToPage;
 };
