ColorPicker: Update Color Picker UI to conform to new design specifications.
Added new widget type, SColorSlider, which has a label, a slider widget, and a spinbox for text entry of color values. Replaced existing FGBA and HSV sliders with this new widget type.
Changed color block to use a rounded brush for the alpha background, and made the border of the brush transparent to blend better with the colors drawn on top of it.
Updated icons on the button that switches between the color wheel and spectrum widgets.
Move Color Themes panel to bottom of window.
Removed Advanced section which hides the color sliders so that they are just always visible.
Updated sizing/spacing of all widgets to conform to the design specifications.
Updated sizing/spacing of inline color picker.

diff --git a/Engine/Source/Editor/DetailCustomizations/Private/ColorStructCustomization.cpp b/Engine/Source/Editor/DetailCustomizations/Private/ColorStructCustomization.cpp
--- a/Engine/Source/Editor/DetailCustomizations/Private/ColorStructCustomization.cpp
+++ b/Engine/Source/Editor/DetailCustomizations/Private/ColorStructCustomization.cpp
@@ -87,7 +87,7 @@ void FColorStructCustomization::MakeHeaderRow(TSharedRef<class IPropertyHandle>&
 	if (InStructPropertyHandle->HasMetaData("InlineColorPicker"))
 	{
 		ColorWidget = CreateInlineColorPicker(StructWeakHandlePtr);
-		ContentWidth = 384.0f;
+		ContentWidth = 492.0f;
 	}
 	else
 	{
diff --git a/Engine/Source/Runtime/AppFramework/Private/Widgets/Colors/SColorPicker.cpp b/Engine/Source/Runtime/AppFramework/Private/Widgets/Colors/SColorPicker.cpp
--- a/Engine/Source/Runtime/AppFramework/Private/Widgets/Colors/SColorPicker.cpp
+++ b/Engine/Source/Runtime/AppFramework/Private/Widgets/Colors/SColorPicker.cpp
@@ -19,6 +19,7 @@
 #include "Widgets/Input/SCheckBox.h"
 #include "Widgets/Input/SSpinBox.h"
 #include "Widgets/Input/SSlider.h"
+#include "Widgets/Colors/SColorSlider.h"
 #include "Widgets/Colors/SComplexGradient.h"
 #include "Widgets/Colors/SSimpleGradient.h"
 #include "Widgets/Colors/SEyeDropperButton.h"
@@ -32,7 +33,7 @@
 
 
 /** A default window size for the color picker which looks nice */
-const FVector2D SColorPicker::DEFAULT_WINDOW_SIZE = FVector2D(441, 537);
+const FVector2D SColorPicker::DEFAULT_WINDOW_SIZE = FVector2D(462, 526);
 
 /** The max time allowed for updating before we shut off auto-updating */
 const double SColorPicker::MAX_ALLOWED_UPDATE_TIME = 0.1;
@@ -98,28 +99,23 @@ void SColorPicker::Construct( const FArguments& InArgs )
 
 	BeginAnimation(FLinearColor(ForceInit), CurrentColorHSV);
 
-	bool bAdvancedSectionExpanded = false;
-
 	if (FPaths::FileExists(GEditorPerProjectIni))
 	{
 		bool WheelMode = true;
 
 		GConfig->GetBool(TEXT("ColorPickerUI"), TEXT("bWheelMode"), WheelMode, GEditorPerProjectIni);
-		GConfig->GetBool(TEXT("ColorPickerUI"), TEXT("bAdvancedSectionExpanded"), bAdvancedSectionExpanded, GEditorPerProjectIni);
 		GConfig->GetBool(TEXT("ColorPickerUI"), TEXT("bSRGBEnabled"), SColorThemesViewer::bSRGBEnabled, GEditorPerProjectIni);
 
 		CurrentMode = WheelMode ? EColorPickerModes::Wheel : EColorPickerModes::Spectrum;
 	}
 
-	bAdvancedSectionExpanded |= InArgs._ExpandAdvancedSection;
-
 	if (bColorPickerIsInlineVersion)
 	{
 		GenerateInlineColorPickerContent();
 	}
 	else
 	{
-		GenerateDefaultColorPickerContent(bAdvancedSectionExpanded);
+		GenerateDefaultColorPickerContent(true /*bAdvancedSectionExpanded*/);
 	}
 }
 
@@ -132,7 +128,7 @@ BEGIN_SLATE_FUNCTION_BUILD_OPTIMIZATION
 void SColorPicker::GenerateDefaultColorPickerContent( bool bAdvancedSectionExpanded )
 {	
 	// The height of the gradient bars beneath the sliders
-	const FSlateFontInfo SmallLayoutFont = FAppStyle::Get().GetFontStyle("ColorPicker.Font");
+	const FSlateFontInfo SmallFont = FAppStyle::Get().GetFontStyle("ColorPicker.SmallFont");
 
 	TSharedPtr<SColorThemesViewer> ThemesViewer = ColorThemesViewer.Pin();
 
@@ -170,65 +166,15 @@ void SColorPicker::GenerateDefaultColorPickerContent( bool bAdvancedSectionExpan
 	[
 		SNew(SVerticalBox)
 
+		// Top Panel, with Color Wheel / Spectrum, Old/New color swatches, and color picker buttons
 		+ SVerticalBox::Slot()
 			.AutoHeight()
 			[
-				SNew(SGridPanel)
-					.FillColumn(0, 1.0f)
+				SNew(SHorizontalBox)
 
-				+ SGridPanel::Slot(0, 0)
-					.Padding(0.0f, 1.0f, 20.0f, 1.0f)
-					[
-						SNew(SHorizontalBox)
-						
-						+ SHorizontalBox::Slot()
-							.FillWidth(0.91f)
-							.Padding(0.0f, 1.0f)
-							[
-								SNew(SOverlay)
-
-								+ SOverlay::Slot()
-									[
-										// color theme bar
-										SAssignNew(CurrentThemeBar, SThemeColorBlocksBar)
-											.ColorTheme(this, &SColorPicker::HandleThemeBarColorTheme)
-											.EmptyText(LOCTEXT("EmptyBarHint", "Drag & drop colors here to save"))
-											.HideTrashCallback(this, &SColorPicker::HideSmallTrash)
-											.ShowTrashCallback(this, &SColorPicker::ShowSmallTrash)
-											.ToolTipText(LOCTEXT("CurrentThemeBarToolTip", "Current Color Theme"))
-											.UseAlpha(SharedThis(this), &SColorPicker::HandleThemeBarUseAlpha)
-											.UseSRGB(SharedThis(this), &SColorPicker::HandleColorPickerUseSRGB)
-											.OnSelectColor(this, &SColorPicker::HandleThemeBarColorSelected)
-									]
-							]
-
-						+ SHorizontalBox::Slot()
-							.FillWidth(0.09f)
-							[
-								// color theme selector
-								SAssignNew(ColorThemeButtonOrSmallTrash, SBorder)
-									.BorderImage(FStyleDefaults::GetNoBrush())
-									.Padding(0.0f)
-							]
-					]
-
-				+ SGridPanel::Slot(1, 0)
-					.HAlign(HAlign_Right)
-					.VAlign(VAlign_Center)
-					[
-						// sRGB check box
-						SNew(SCheckBox)
-							.ToolTipText(LOCTEXT("SRGBCheckboxToolTip", "Toggle gamma corrected sRGB previewing"))
-							.IsChecked(this, &SColorPicker::HandleSRGBCheckBoxIsChecked)
-							.OnCheckStateChanged(this, &SColorPicker::HandleSRGBCheckBoxCheckStateChanged)
-							[
-								SNew(STextBlock)
-									.Text(LOCTEXT("SRGBCheckboxLabel", "sRGB Preview"))
-							]
-					]
-
-				+ SGridPanel::Slot(0, 1)
-					.Padding(0.0f, 8.0f, 20.0f, 0.0f)
+				// Overlay displaying either the Color Wheel with Saturation and Value vertical sliders, or just the Color Spectrum widget
+				+ SHorizontalBox::Slot()
+					.AutoWidth()
 					[
 						SNew(SBorder)
 							.BorderImage(FAppStyle::Get().GetBrush("NoBorder"))
@@ -243,8 +189,8 @@ void SColorPicker::GenerateDefaultColorPickerContent( bool bAdvancedSectionExpan
 										SNew(SHorizontalBox)
 
 										+ SHorizontalBox::Slot()
-											.FillWidth(1.0f)
-											.HAlign(HAlign_Center)
+											.MinWidth(200.0f)
+											.MaxWidth(200.0f)
 											[
 												SNew(SColorWheel)
 													.SelectedColor(this, &SColorPicker::GetCurrentColor)
@@ -256,7 +202,7 @@ void SColorPicker::GenerateDefaultColorPickerContent( bool bAdvancedSectionExpan
 
 										+ SHorizontalBox::Slot()
 											.AutoWidth()
-											.Padding(4.0f, 0.0f)
+											.Padding(20.0f, 0.0f, 0.0f, 0.0f)
 											[
 												// saturation slider
 												MakeColorSlider(EColorPickerChannels::Saturation)
@@ -264,6 +210,7 @@ void SColorPicker::GenerateDefaultColorPickerContent( bool bAdvancedSectionExpan
 
 										+ SHorizontalBox::Slot()
 											.AutoWidth()
+											.Padding(20.0f, 0.0f, 0.0f, 0.0f)
 											[
 												// value slider
 												MakeColorSlider(EColorPickerChannels::Value)
@@ -275,7 +222,7 @@ void SColorPicker::GenerateDefaultColorPickerContent( bool bAdvancedSectionExpan
 									[
 										SNew(SBox)
 											.HeightOverride(200.0f)
-											.WidthOverride(292.0f)
+											.WidthOverride(304.0f)
 											[
 												SNew(SColorSpectrum)
 													.SelectedColor(this, &SColorPicker::GetCurrentColor)
@@ -288,46 +235,74 @@ void SColorPicker::GenerateDefaultColorPickerContent( bool bAdvancedSectionExpan
 							]
 					]
 
-				+ SGridPanel::Slot(1, 1)
-					.Padding(0.0f, 8.0f, 0.0f, 0.0f)
+				+ SHorizontalBox::Slot()
+					.AutoWidth()
+					.Padding(20.0f, 0.0f, 0.0f, 0.0f)
 					[
 						SNew(SVerticalBox)
 
 						+ SVerticalBox::Slot()
 							.AutoHeight()
 							[
-								SNew(SBox)
-									.HeightOverride(100.0f)
-									.WidthOverride(70.0f)
+								// sRGB check box
+								SNew(SCheckBox)
+									.ToolTipText(LOCTEXT("SRGBCheckboxToolTip", "Toggle gamma corrected sRGB previewing"))
+									.IsChecked(this, &SColorPicker::HandleSRGBCheckBoxIsChecked)
+									.OnCheckStateChanged(this, &SColorPicker::HandleSRGBCheckBoxCheckStateChanged)
 									[
-										// color preview
-										MakeColorPreviewBox()
+										SNew(STextBlock)
+											.Text(LOCTEXT("SRGBCheckboxLabel", "sRGB Preview"))
 									]
 							]
 
 						+ SVerticalBox::Slot()
 							.AutoHeight()
-							.Padding(0.0f, 16.0f, 0.0f, 0.0f)
-							.VAlign(VAlign_Top)
+							.Padding(0.0f, 8.0f, 0.0f, 0.0f)
+							[
+								// color preview
+								MakeColorPreviewBox()
+							]
+
+						+ SVerticalBox::Slot()
+							.MinHeight(28.0f)
+							.MaxHeight(28.0f)
+							.Padding(0.0f, 8.0f, 0.0f, 0.0f)
 							[
 								SNew(SHorizontalBox)
 
 								+ SHorizontalBox::Slot()
-									.HAlign(HAlign_Left)
+									.MinWidth(48.0f)
+									.MaxWidth(48.0f)
 									[
 										// mode selector
 										SNew(SButton)
 											.OnClicked(this, &SColorPicker::HandleColorPickerModeButtonClicked)
+											.ContentPadding(FMargin(2.0f, 2.5f))
 											.Content()
 											[
-												SNew(SImage)
-													.Image(FAppStyle::Get().GetBrush("ColorPicker.Mode"))
+												SNew(SOverlay)
 													.ToolTipText(LOCTEXT("ColorPickerModeEToolTip", "Toggle between color wheel and color spectrum."))
-											]									
+
+												+ SOverlay::Slot()
+													[
+														SNew(SImage)
+															.Image(FAppStyle::Get().GetBrush("ColorPicker.ModeWheel"))
+															.Visibility(this, &SColorPicker::HandleColorPickerModeVisibility, EColorPickerModes::Spectrum)
+													]
+
+												+ SOverlay::Slot()
+													[
+														SNew(SImage)
+															.Image(FAppStyle::Get().GetBrush("ColorPicker.ModeSpectrum"))
+															.Visibility(this, &SColorPicker::HandleColorPickerModeVisibility, EColorPickerModes::Wheel)
+													]
+											]
 									]
 
 								+ SHorizontalBox::Slot()
-									.HAlign(HAlign_Right)
+									.MinWidth(48.0f)
+									.MaxWidth(48.0f)
+									.Padding(10.0f, 0.0f, 0.0f, 0.0f)
 									[
 										// eye dropper
 										SNew(SEyeDropperButton)
@@ -341,147 +316,176 @@ void SColorPicker::GenerateDefaultColorPickerContent( bool bAdvancedSectionExpan
 					]
 			]
 
-		// advanced settings
+		// Color Sliders Panel
 		+ SVerticalBox::Slot()
 			.AutoHeight()
-			.Padding(0.0f, 4.0f, 0.0f, 0.0f)
+			.Padding(0.0f, 16.0f, 0.0f, 0.0f)
 			[
-				SNew(SExpandableArea)
-					.AreaTitle(LOCTEXT("AdvancedAreaTitle", "Advanced"))
-					.BorderBackgroundColor(FLinearColor::Transparent)
-					.InitiallyCollapsed(!bAdvancedSectionExpanded)
-					.OnAreaExpansionChanged(this, &SColorPicker::HandleAdvancedAreaExpansionChanged)
-					.Padding(FMargin(0.0f, 1.0f, 0.0f, 8.0f))
-					.BodyContent()
+				SNew(SHorizontalBox)
+
+				// RGBA Color Sliders
+				+ SHorizontalBox::Slot()
 					[
-						SNew(SHorizontalBox)
+						SNew(SVerticalBox)
 
-						// RGBA inputs
-						+ SHorizontalBox::Slot()
-							.Padding(0.0f, 0.0f, 4.0f, 0.0f)
+						+ SVerticalBox::Slot()
+							.AutoHeight()
 							[
-								SNew(SVerticalBox)
+								MakeColorSpinBox(EColorPickerChannels::Red)
+							]
 
-								// Red
-								+ SVerticalBox::Slot()
-									[
-										MakeColorSpinBox(EColorPickerChannels::Red)
-									]
+						+ SVerticalBox::Slot()
+							.AutoHeight()
+							.Padding(0.0f, 8.0f, 0.0f, 0.0f)
+							[
+								MakeColorSpinBox(EColorPickerChannels::Green)
+							]
 
-								// Green
-								+ SVerticalBox::Slot()
-									.Padding(0.0f, 8.0f, 0.0f, 0.0f)
-									[
-										MakeColorSpinBox(EColorPickerChannels::Green)
-									]
+						+ SVerticalBox::Slot()
+							.AutoHeight()
+							.Padding(0.0f, 8.0f, 0.0f, 0.0f)
+							[
+								MakeColorSpinBox(EColorPickerChannels::Blue)
+							]
 
-								// Blue
-								+ SVerticalBox::Slot()
-									.Padding(0.0f, 8.0f, 0.0f, 0.0f)
+						+ SVerticalBox::Slot()
+							.AutoHeight()
+							.Padding(0.0f, 8.0f, 0.0f, 0.0f)
+							[
+								MakeColorSpinBox(EColorPickerChannels::Alpha)
+							]
+					]
+
+				// HSV Color Sliders and & Hexadecimal TextBoxes
+				+ SHorizontalBox::Slot()
+					.Padding(16.0f, 0.0f, 0.0f, 0.0f)
+					[
+						SNew(SVerticalBox)
+							
+						+ SVerticalBox::Slot()
+							.AutoHeight()
+							[
+								MakeColorSpinBox(EColorPickerChannels::Hue)
+							]
+
+						+ SVerticalBox::Slot()
+							.AutoHeight()
+							.Padding(0.0f, 8.0f, 0.0f, 0.0f)
+							[
+								MakeColorSpinBox(EColorPickerChannels::Saturation)
+							]
+
+						+ SVerticalBox::Slot()
+							.AutoHeight()
+							.Padding(0.0f, 8.0f, 0.0f, 0.0f)
+							[
+								MakeColorSpinBox(EColorPickerChannels::Value)
+							]
+
+						// Hexadecimal Linear TextBox
+						+ SVerticalBox::Slot()
+							.MinHeight(20)
+							.MaxHeight(20)
+							.Padding(0.0f, 8.0f, 0.0f, 0.0f)
+							[
+								SNew(SHorizontalBox)
+									.ToolTipText(LOCTEXT("HexLinearSliderToolTip", "Hexadecimal Linear Value"))
+
+								+ SHorizontalBox::Slot()
+									.HAlign(HAlign_Left)
+									.VAlign(VAlign_Center)
 									[
-										MakeColorSpinBox(EColorPickerChannels::Blue)
+										SNew(STextBlock)
+											.Text(LOCTEXT("HexLinearInputLabel", "Hex Linear"))
 									]
 
-								// Alpha
-								+ SVerticalBox::Slot()
-									.Padding(0.0f, 8.0f, 0.0f, 0.0f)
+								+ SHorizontalBox::Slot()
+									.AutoWidth()
+									.HAlign(HAlign_Right)
 									[
-										MakeColorSpinBox(EColorPickerChannels::Alpha)
+										SNew(SEditableTextBox)
+											.MinDesiredWidth(114.0f)
+											.Text(this, &SColorPicker::HandleHexLinearBoxText)
+											.Font(SmallFont)
+											.Padding(FMargin(8.0f, 4.0f, 8.0f, 4.0f))
+											.OnTextCommitted(this, &SColorPicker::HandleHexLinearInputTextCommitted)
 									]
 							]
 
-						// HSV & Hex inputs
-						+ SHorizontalBox::Slot()
-							.Padding(4.0f, 0.0f, 0.0f, 0.0f)
+						// Hexadecimal sRGB TextBox
+						+ SVerticalBox::Slot()
+							.MinHeight(20)
+							.MaxHeight(20)
+							.Padding(0.0f, 8.0f, 0.0f, 0.0f)
 							[
-								SNew(SVerticalBox)
-							
-								// Hue
-								+ SVerticalBox::Slot()
-								[
-									MakeColorSpinBox(EColorPickerChannels::Hue)
-								]
-
-								// Saturation
-								+ SVerticalBox::Slot()
-								.Padding(0.0f, 8.0f, 0.0f, 0.0f)
-								[
-									MakeColorSpinBox(EColorPickerChannels::Saturation)
-								]
-
-								// Value
-								+ SVerticalBox::Slot()
-								.Padding(0.0f, 8.0f, 0.0f, 0.0f)
-								[
-									MakeColorSpinBox(EColorPickerChannels::Value)
-								]
-
-								// Hex linear
-								+ SVerticalBox::Slot()
-									.HAlign(HAlign_Right)
-									.VAlign(VAlign_Top)
-									.Padding(0.0f, 12.0f, 0.0f, 0.0f)
-									[
-										SNew(SHorizontalBox)
-											.ToolTipText(LOCTEXT("HexLinearSliderToolTip", "Hexadecimal Linear Value"))
-
-										+ SHorizontalBox::Slot()
-											.AutoWidth()
-											.Padding(0.0f, 0.0f, 4.0f, 0.0f)
-											.VAlign(VAlign_Center)
-											[
-												SNew(STextBlock)
-													.Text(LOCTEXT("HexLinearInputLabel", "Hex Linear"))
-											]
+								SNew(SHorizontalBox)
+									.ToolTipText(LOCTEXT("HexSRGBSliderToolTip", "Hexadecimal sRGB Value"))
 
-										+ SHorizontalBox::Slot()
-											.AutoWidth()
-											.MaxWidth(90.0f)
-											[
-												SNew(SEditableTextBox)
-													.MinDesiredWidth(90.0f)
-													.Text(this, &SColorPicker::HandleHexLinearBoxText)
-													.OnTextCommitted(this, &SColorPicker::HandleHexLinearInputTextCommitted)
-											]
-								]
+								+ SHorizontalBox::Slot()
+									.HAlign(HAlign_Left)
+									.VAlign(VAlign_Center)
+									[
+										SNew(STextBlock)
+											.Text(LOCTEXT("HexSRGBInputLabel", "Hex sRGB"))
+									]
 
-								// Hex sRGB
-								+ SVerticalBox::Slot()
+								+ SHorizontalBox::Slot()
+									.AutoWidth()
 									.HAlign(HAlign_Right)
-									.VAlign(VAlign_Top)
-									.Padding(0.0f, 8.0f, 0.0f, 0.0f)
 									[
-										SNew(SHorizontalBox)
-										.ToolTipText(LOCTEXT("HexSRGBSliderToolTip", "Hexadecimal sRGB Value"))
+										SNew(SEditableTextBox)
+											.MinDesiredWidth(114.0f)
+											.Text(this, &SColorPicker::HandleHexSRGBBoxText)
+											.Font(SmallFont)
+											.Padding(FMargin(8.0f, 4.0f, 8.0f, 4.0f))
+											.OnTextCommitted(this, &SColorPicker::HandleHexSRGBInputTextCommitted)
+									]
+							]
+					]
+			]
 
-										+ SHorizontalBox::Slot()
-											.AutoWidth()
-											.Padding(0.0f, 0.0f, 4.0f, 0.0f)
-											.VAlign(VAlign_Center)
-											[
-												SNew(STextBlock)
-												.Text(LOCTEXT("HexSRGBInputLabel", "Hex sRGB"))
-											]
+		// Color Themes Panel
+		+ SVerticalBox::Slot()
+			.AutoHeight()
+			.Padding(0.0f, 16.0f, 0.0f, 0.0f)
+			[
+				SNew(SHorizontalBox)
 
-										+ SHorizontalBox::Slot()
-											.AutoWidth()
-											.MaxWidth(90.0f)
-											[
-												SNew(SEditableTextBox)
-												.MinDesiredWidth(90.0f)
-												.Text(this, &SColorPicker::HandleHexSRGBBoxText)
-												.OnTextCommitted(this, &SColorPicker::HandleHexSRGBInputTextCommitted)
-											]
-								]
-						]
-				]
-			]
+				+ SHorizontalBox::Slot()
+					.FillWidth(0.91f)
+					[
+						SNew(SOverlay)
+
+						+ SOverlay::Slot()
+							[
+								// color theme bar
+								SAssignNew(CurrentThemeBar, SThemeColorBlocksBar)
+									.ColorTheme(this, &SColorPicker::HandleThemeBarColorTheme)
+									.EmptyText(LOCTEXT("EmptyBarHint", "Drag & drop colors here to save"))
+									.HideTrashCallback(this, &SColorPicker::HideSmallTrash)
+									.ShowTrashCallback(this, &SColorPicker::ShowSmallTrash)
+									.ToolTipText(LOCTEXT("CurrentThemeBarToolTip", "Current Color Theme"))
+									.UseAlpha(SharedThis(this), &SColorPicker::HandleThemeBarUseAlpha)
+									.UseSRGB(SharedThis(this), &SColorPicker::HandleColorPickerUseSRGB)
+									.OnSelectColor(this, &SColorPicker::HandleThemeBarColorSelected)
+							]
+					]
 
+				+ SHorizontalBox::Slot()
+					.FillWidth(0.09f)
+					[
+						// color theme selector
+						SAssignNew(ColorThemeButtonOrSmallTrash, SBorder)
+							.BorderImage(FStyleDefaults::GetNoBrush())
+							.Padding(0.0f)
+					]
+			]
+			
 		// dialog buttons
 		+ SVerticalBox::Slot()
 			.AutoHeight()
 			.HAlign(HAlign_Right)
-			.Padding(0.0f, 12.0f, 0.0f, 0.0f)
+			.Padding(0.0f, 16.0f, 0.0f, 0.0f)
 			[
 				SNew(SUniformGridPanel)
 					.MinDesiredSlotHeight(FAppStyle::Get().GetFloat("StandardDialog.MinDesiredSlotHeight"))
@@ -558,8 +562,7 @@ void SColorPicker::GenerateInlineColorPickerContent()
 		SNew(SHorizontalBox)
 
 		+ SHorizontalBox::Slot()
-			.FillWidth(1.0f)
-			.HAlign(HAlign_Center)
+			.AutoWidth()
 			[
 				SNew(SColorWheel)
 					.SelectedColor(this, &SColorPicker::GetCurrentColor)
@@ -571,7 +574,7 @@ void SColorPicker::GenerateInlineColorPickerContent()
 
 		+ SHorizontalBox::Slot()
 			.AutoWidth()
-			.Padding(4.0f, 0.0f)
+			.Padding(8.0f, 0.0f, 0.0f, 0.0f)
 			[
 				// saturation slider
 				MakeColorSlider(EColorPickerChannels::Saturation)
@@ -579,6 +582,7 @@ void SColorPicker::GenerateInlineColorPickerContent()
 
 		+ SHorizontalBox::Slot()
 			.AutoWidth()
+			.Padding(8.0f, 0.0f, 0.0f, 0.0f)
 			[
 				// value slider
 				MakeColorSlider(EColorPickerChannels::Value)
@@ -586,6 +590,7 @@ void SColorPicker::GenerateInlineColorPickerContent()
 
 		+ SHorizontalBox::Slot()
 			.AutoWidth()
+			.Padding(8.0f, 0.0f, 0.0f, 0.0f)
 			[
 				// Alpha slider
 				AlphaSlider
@@ -751,31 +756,19 @@ TSharedRef<SWidget> SColorPicker::MakeColorSlider( EColorPickerChannels Channel
 		return SNullWidget::NullWidget;
 	}
 
-	return SNew(SOverlay)
+	return SNew(SColorSlider)
+		.Orientation(Orient_Vertical)
 		.ToolTipText(SliderTooltip)
-
-	+ SOverlay::Slot()
-		.Padding(FMargin(4.0f, 0.0f))
-		[
-			SNew(SSimpleGradient)
-				.EndColor(this, &SColorPicker::HandleColorSliderEndColor, Channel)
-				.StartColor(this, &SColorPicker::HandleColorSliderStartColor, Channel)
-				.Orientation(Orient_Horizontal)
-				//.UseSRGB(SharedThis(this), &SColorPicker::HandleColorPickerUseSRGB)
-		]
-
-	+ SOverlay::Slot()
-		[
-			SNew(SSlider)
-				.IndentHandle(false)
-				.Orientation(Orient_Vertical)
-				.SliderBarColor(FLinearColor::Transparent)
-				.Style(&FAppStyle::Get().GetWidgetStyle<FSliderStyle>("ColorPicker.Slider"))
-				.Value(this, &SColorPicker::HandleColorSpinBoxValue, Channel)
-				.OnMouseCaptureBegin(const_cast<SColorPicker*>(this), &SColorPicker::HandleInteractiveChangeBegin)
-				.OnMouseCaptureEnd(const_cast<SColorPicker*>(this), &SColorPicker::HandleInteractiveChangeEnd)
-				.OnValueChanged(const_cast<SColorPicker*>(this), &SColorPicker::HandleColorSpinBoxValueChanged, Channel)
-		];
+		.MinSliderValue(0.0f)
+		.MaxSliderValue(Channel == EColorPickerChannels::Hue ? 359.0f : 1.0f)
+		.Delta(Channel == EColorPickerChannels::Hue ? 1.0f : 0.001f)
+		.HasAlphaBackground(Channel == EColorPickerChannels::Alpha)
+		.GradientColors(this, &SColorPicker::GetGradientColors, Channel)
+		.Value(this, &SColorPicker::HandleColorSpinBoxValue, Channel)
+		.Visibility(this, &SColorPicker::HandleColorPickerModeVisibility, EColorPickerModes::Wheel)
+		.OnBeginSliderMovement(const_cast<SColorPicker*>(this), &SColorPicker::HandleInteractiveChangeBegin)
+		.OnEndSliderMovement(const_cast<SColorPicker*>(this), &SColorPicker::HandleInteractiveChangeEnd)
+		.OnValueChanged(const_cast<SColorPicker*>(this), &SColorPicker::HandleColorSpinBoxValueChanged, Channel);
 }
 
 
@@ -786,35 +779,8 @@ TSharedRef<SWidget> SColorPicker::MakeColorSpinBox( EColorPickerChannels Channel
 		return SNullWidget::NullWidget;
 	}
 
-	const int32 GradientHeight = 6;
 	const float HDRMaxValue = bClampValue ? 1.f : FLT_MAX;
-	const FSlateFontInfo SmallLayoutFont = FAppStyle::Get().GetFontStyle("ColorPicker.Font");
-
-	// create gradient widget
-	TSharedPtr<SWidget> GradientWidget;
-	
-	if (Channel == EColorPickerChannels::Hue)
-	{
-		TArray<FLinearColor> HueGradientColors;
 
-		for (int32 i = 0; i < 7; ++i)
-		{
-			HueGradientColors.Add( FLinearColor((i % 6) * 60.f, 1.f, 1.f).HSVToLinearRGB() );
-		}
-
-		GradientWidget = SNew(SComplexGradient)
-			.GradientColors(HueGradientColors);
-	}
-	else
-	{
-		GradientWidget = SNew(SSimpleGradient)
-			.StartColor(this, &SColorPicker::GetGradientStartColor, Channel)
-			.EndColor(this, &SColorPicker::GetGradientEndColor, Channel)
-			.HasAlphaBackground(Channel == EColorPickerChannels::Alpha);
-			//.UseSRGB(SharedThis(this), &SColorPicker::HandleColorPickerUseSRGB);
-	}
-	
-	// create spin box
 	float MaxValue;
 	FText SliderLabel;
 	FText SliderTooltip;
@@ -867,175 +833,115 @@ TSharedRef<SWidget> SColorPicker::MakeColorSpinBox( EColorPickerChannels Channel
 		return SNullWidget::NullWidget;
 	}
 
-	// Define a maximum size for the spin box containers to prevent them from stretching out the color picker window
-	static const float MaxSpinBoxSize = 192.0f;
-
-	return SNew(SHorizontalBox)
-
-	+ SHorizontalBox::Slot()
-		.AutoWidth()
-		.Padding(0.0f, 0.0f, 4.0f, 0.0f)
-		.VAlign(VAlign_Center)
-		[
-			SNew(STextBlock)
-				.Text(SliderLabel)
-		]
-
-	+ SHorizontalBox::Slot()
-		.FillWidth(1.0f)
-		.MaxWidth(MaxSpinBoxSize)
-		[
-			SNew(SVerticalBox)
-				.ToolTipText(SliderTooltip)
-
-			+ SVerticalBox::Slot()
-				.AutoHeight()
-				[
-					SNew(SSpinBox<float>)
-						.MinValue(0.0f)
-						.MaxValue(MaxValue)
-						.MinSliderValue(0.0f)
-						.MaxSliderValue(Channel == EColorPickerChannels::Hue ? 359.0f : 1.0f)
-						.Delta(Channel == EColorPickerChannels::Hue ? 1.0f : 0.001f)
-						.Font(SmallLayoutFont)
-						.Value(this, &SColorPicker::HandleColorSpinBoxValue, Channel)
-						.OnBeginSliderMovement(const_cast<SColorPicker*>(this), &SColorPicker::HandleInteractiveChangeBegin)
-						.OnEndSliderMovement(const_cast<SColorPicker*>(this), &SColorPicker::HandleInteractiveChangeEnd)
-						.OnValueChanged(const_cast<SColorPicker*>(this), &SColorPicker::HandleColorSpinBoxValueChanged, Channel)
-				]
-
-			+ SVerticalBox::Slot()
-				.AutoHeight()
-				[
-					SNew(SBox)
-						.HeightOverride(GradientHeight)
-						[
-							GradientWidget.ToSharedRef()
-						]
-				]
-		];
+	return SNew(SColorSlider)
+		.Label(SliderLabel)
+		.ToolTipText(SliderTooltip)
+		.MinSpinBoxValue(0.0f)
+		.MaxSpinBoxValue(MaxValue)
+		.MinSliderValue(0.0f)
+		.MaxSliderValue(Channel == EColorPickerChannels::Hue ? 359.0f : 1.0f)
+		.Delta(Channel == EColorPickerChannels::Hue ? 1.0f : 0.001f)
+		.HasAlphaBackground(Channel == EColorPickerChannels::Alpha)
+		.GradientColors(this, &SColorPicker::GetGradientColors, Channel)
+		.Value(this, &SColorPicker::HandleColorSpinBoxValue, Channel)
+		.OnBeginSliderMovement(const_cast<SColorPicker*>(this), &SColorPicker::HandleInteractiveChangeBegin)
+		.OnEndSliderMovement(const_cast<SColorPicker*>(this), &SColorPicker::HandleInteractiveChangeEnd)
+		.OnBeginSpinBoxMovement(const_cast<SColorPicker*>(this), &SColorPicker::HandleInteractiveChangeBegin)
+		.OnEndSpinBoxMovement(const_cast<SColorPicker*>(this), &SColorPicker::HandleInteractiveChangeEnd)
+		.OnValueChanged(const_cast<SColorPicker*>(this), &SColorPicker::HandleColorSpinBoxValueChanged, Channel);
 }
 
 
 TSharedRef<SWidget> SColorPicker::MakeColorPreviewBox() const
 {
-	return SNew(SOverlay)
+	const float OpaqueWidth = bUseAlpha.Get() ? 53.0f : 106.0f;
+	const FVector4 OldOpaqueBlockCornerRadius = bUseAlpha.Get() ? FVector4(4.0, 0.0, 0.0, 0.0) : FVector4(4.0, 4.0, 0.0, 0.0);
+	const FVector4 NewOpaqueBlockCornerRadius = bUseAlpha.Get() ? FVector4(0.0, 0.0, 0.0, 4.0) : FVector4(0.0, 0.0, 4.0, 4.0);
 
-	+ SOverlay::Slot()
+	return SNew(SVerticalBox)
+
+	+ SVerticalBox::Slot()
+		.AutoHeight()
 		[
-			// preview blocks
-			SNew(SVerticalBox)
+			SNew(SOverlay)
 
-			+ SVerticalBox::Slot()
-				.AutoHeight()
+			+ SOverlay::Slot()
 				[
-					SNew(STextBlock)
-						.Text(LOCTEXT("OldColorLabel", "Old"))
+					// old color (alpha)
+					SNew(SColorBlock)
+						.ColorIsHSV(true)
+						.ShowBackgroundForAlpha(true)
+						.ToolTipText(LOCTEXT("OldColorAlphaToolTip", "Old color with alpha (drag to theme bar to save)"))
+						.Color(OldColor)
+						.Visibility(SharedThis(this), &SColorPicker::HandleAlphaColorBlockVisibility)
+						.OnMouseButtonDown(const_cast<SColorPicker*>(this), &SColorPicker::HandleOldColorBlockMouseButtonDown, true)
+						.UseSRGB(SharedThis(this), &SColorPicker::HandleColorPickerUseSRGB)
+						.Cursor(EMouseCursor::GrabHand)
+						.Size(FVector2D(106.0, 32.0))
+						.CornerRadius(FVector4(4.0, 4.0, 0.0, 0.0))
 				]
 
-			+ SVerticalBox::Slot()
+			+ SOverlay::Slot()
 				[
 					SNew(SHorizontalBox)
 
-					+ SHorizontalBox::Slot()
+					+SHorizontalBox::Slot()
+						.AutoWidth()
 						[
 							// old color
-							SNew(SColorBlock) 
-								.ColorIsHSV(true) 
+							SNew(SColorBlock)
+								.ColorIsHSV(true)
 								.AlphaDisplayMode(EColorBlockAlphaDisplayMode::Ignore)
 								.ToolTipText(LOCTEXT("OldColorToolTip", "Old color without alpha (drag to theme bar to save)"))
-								.Color(OldColor) 
+								.Color(OldColor)
 								.OnMouseButtonDown(const_cast<SColorPicker*>(this), &SColorPicker::HandleOldColorBlockMouseButtonDown, false)
 								.UseSRGB(SharedThis(this), &SColorPicker::HandleColorPickerUseSRGB)
 								.Cursor(EMouseCursor::GrabHand)
+								.Size(FVector2D(OpaqueWidth, 32.0))
+								.CornerRadius(OldOpaqueBlockCornerRadius)
 						]
+				]
+		]
 
-					+ SHorizontalBox::Slot()
-						[
-							// old color (alpha)
-							SNew(SColorBlock) 
-								.ColorIsHSV(true) 
-								.ShowBackgroundForAlpha(true)
-								.ToolTipText(LOCTEXT("OldColorAlphaToolTip", "Old color with alpha (drag to theme bar to save)"))
-								.Color(OldColor) 
-								.Visibility(SharedThis(this), &SColorPicker::HandleAlphaColorBlockVisibility)
-								.OnMouseButtonDown(const_cast<SColorPicker*>(this), &SColorPicker::HandleOldColorBlockMouseButtonDown, true)
-								.UseSRGB(SharedThis(this), &SColorPicker::HandleColorPickerUseSRGB)
-								.Cursor(EMouseCursor::GrabHand)
-						]
+	+ SVerticalBox::Slot()
+		.AutoHeight()
+		[
+			SNew(SOverlay)
+
+			+ SOverlay::Slot()
+				[
+					// new color (alpha)
+					SNew(SColorBlock)
+						.ColorIsHSV(true)
+						.ShowBackgroundForAlpha(true)
+						.ToolTipText(LOCTEXT("NewColorAlphaToolTip", "New color with alpha (drag to theme bar to save)"))
+						.Color(this, &SColorPicker::GetCurrentColor)
+						.Visibility(SharedThis(this), &SColorPicker::HandleAlphaColorBlockVisibility)
+						.OnMouseButtonDown(const_cast<SColorPicker*>(this), &SColorPicker::HandleNewColorBlockMouseButtonDown, true)
+						.UseSRGB(SharedThis(this), &SColorPicker::HandleColorPickerUseSRGB)
+						.Cursor(EMouseCursor::GrabHand)
+						.Size(FVector2D(106.0, 32.0))
+						.CornerRadius(FVector4(0.0, 0.0, 4.0, 4.0))
 				]
 
-			+ SVerticalBox::Slot()
+			+ SOverlay::Slot()
 				[
 					SNew(SHorizontalBox)
 
 					+ SHorizontalBox::Slot()
+						.AutoWidth()
 						[
 							// new color
-							SNew(SColorBlock) 
-								.ColorIsHSV(true) 
+							SNew(SColorBlock)
+								.ColorIsHSV(true)
 								.AlphaDisplayMode(EColorBlockAlphaDisplayMode::Ignore)
 								.ToolTipText(LOCTEXT("NewColorToolTip", "New color without alpha (drag to theme bar to save)"))
 								.Color(this, &SColorPicker::GetCurrentColor)
 								.OnMouseButtonDown(const_cast<SColorPicker*>(this), &SColorPicker::HandleNewColorBlockMouseButtonDown, false)
 								.UseSRGB(SharedThis(this), &SColorPicker::HandleColorPickerUseSRGB)
 								.Cursor(EMouseCursor::GrabHand)
-						]
-
-					+ SHorizontalBox::Slot()
-						[
-							// new color (alpha)
-							SNew(SColorBlock) 
-								.ColorIsHSV(true) 
-								.ShowBackgroundForAlpha(true)
-								.ToolTipText(LOCTEXT("NewColorAlphaToolTip", "New color with alpha (drag to theme bar to save)"))
-								.Color(this, &SColorPicker::GetCurrentColor)
-								.Visibility(SharedThis(this), &SColorPicker::HandleAlphaColorBlockVisibility)
-								.OnMouseButtonDown(const_cast<SColorPicker*>(this), &SColorPicker::HandleNewColorBlockMouseButtonDown, true)
-								.UseSRGB(SharedThis(this), &SColorPicker::HandleColorPickerUseSRGB)
-								.Cursor(EMouseCursor::GrabHand)
-						]
-				]
-
-			+ SVerticalBox::Slot()
-				.AutoHeight()
-				[
-					SNew(STextBlock)
-						.Text(LOCTEXT("NewColorLabel", "New"))
-				]
-		]
-
-	+ SOverlay::Slot()
-		.VAlign(VAlign_Center)
-		[
-			// block separators
-			SNew(SHorizontalBox)
-
-			+ SHorizontalBox::Slot()
-				.FillWidth(0.5f)
-				.HAlign(HAlign_Left)
-				[
-					SNew(SBox)
-						.HeightOverride(2.0f)
-						.WidthOverride(4.0f)
-						[
-							SNew(SBorder)
-								.BorderImage(FAppStyle::Get().GetBrush("ColorPicker.Separator"))
-								.Padding(0.0f)
-						]
-				]
-
-			+ SHorizontalBox::Slot()
-				.FillWidth(0.5f)
-				.HAlign(HAlign_Right)
-				[
-					SNew(SBox)
-						.HeightOverride(2.0f)
-						.WidthOverride(4.0f)
-						[
-							SNew(SBorder)
-								.BorderImage(FAppStyle::Get().GetBrush("ColorPicker.Separator"))
-								.Padding(0.0f)
+								.Size(FVector2D(OpaqueWidth, 32.0))
+								.CornerRadius(NewOpaqueBlockCornerRadius)
 						]
 				]
 		];
@@ -1050,12 +956,12 @@ FLinearColor SColorPicker::GetGradientEndColor( EColorPickerChannels Channel ) c
 {
 	switch (Channel)
 	{
-	case EColorPickerChannels::Red:			return FLinearColor(1.0f, CurrentColorRGB.G, CurrentColorRGB.B, 1.0f);
-	case EColorPickerChannels::Green:		return FLinearColor(CurrentColorRGB.R, 1.0f, CurrentColorRGB.B, 1.0f);
-	case EColorPickerChannels::Blue:		return FLinearColor(CurrentColorRGB.R, CurrentColorRGB.G, 1.0f, 1.0f);
-	case EColorPickerChannels::Alpha:		return FLinearColor(CurrentColorRGB.R, CurrentColorRGB.G, CurrentColorRGB.B, 1.0f);
+	case EColorPickerChannels::Red:			return FLinearColor::Red;
+	case EColorPickerChannels::Green:		return FLinearColor::Green;
+	case EColorPickerChannels::Blue:		return FLinearColor::Blue;
+	case EColorPickerChannels::Alpha:		return FLinearColor::White;
 	case EColorPickerChannels::Saturation:	return FLinearColor(CurrentColorHSV.R, 1.0f, CurrentColorHSV.B, 1.0f).HSVToLinearRGB();
-	case EColorPickerChannels::Value:		return FLinearColor(CurrentColorHSV.R, CurrentColorHSV.G, 1.0f, 1.0f).HSVToLinearRGB();
+	case EColorPickerChannels::Value:		return FLinearColor::White;
 	default:								return FLinearColor();
 	}
 }
@@ -1065,23 +971,32 @@ FLinearColor SColorPicker::GetGradientStartColor( EColorPickerChannels Channel )
 {
 	switch (Channel)
 	{
-	case EColorPickerChannels::Red:			return FLinearColor(0.0f, CurrentColorRGB.G, CurrentColorRGB.B, 1.0f);
-	case EColorPickerChannels::Green:		return FLinearColor(CurrentColorRGB.R, 0.0f, CurrentColorRGB.B, 1.0f);
-	case EColorPickerChannels::Blue:		return FLinearColor(CurrentColorRGB.R, CurrentColorRGB.G, 0.0f, 1.0f);
-	case EColorPickerChannels::Alpha:		return FLinearColor(CurrentColorRGB.R, CurrentColorRGB.G, CurrentColorRGB.B, 0.0f);
+	case EColorPickerChannels::Red:			return FLinearColor::Black;
+	case EColorPickerChannels::Green:		return FLinearColor::Black;
+	case EColorPickerChannels::Blue:		return FLinearColor::Black;
+	case EColorPickerChannels::Alpha:		return FLinearColor::Transparent;
 	case EColorPickerChannels::Saturation:	return FLinearColor(CurrentColorHSV.R, 0.0f, CurrentColorHSV.B, 1.0f).HSVToLinearRGB();
-	case EColorPickerChannels::Value:		return FLinearColor(CurrentColorHSV.R, CurrentColorHSV.G, 0.0f, 1.0f).HSVToLinearRGB();
+	case EColorPickerChannels::Value:		return FLinearColor::Black;
 	default:								return FLinearColor();
 	}
 }
 
-
-void SColorPicker::HandleAdvancedAreaExpansionChanged( bool Expanded )
+TArray<FLinearColor> SColorPicker::GetGradientColors( EColorPickerChannels Channel ) const
 {
-	if (FPaths::FileExists(GEditorPerProjectIni))
+	TArray<FLinearColor> Colors;
+	if (Channel == EColorPickerChannels::Hue)
+	{
+		for (int32 i = 0; i < 7; ++i)
+		{
+			Colors.Add(FLinearColor((i % 6) * 60.f, 1.f, 1.f).HSVToLinearRGB());
+		}
+	}
+	else
 	{
-		GConfig->SetBool(TEXT("ColorPickerUI"), TEXT("bAdvancedSectionExpanded"), Expanded, GEditorPerProjectIni);
+		Colors.Add(GetGradientStartColor(Channel));
+		Colors.Add(GetGradientEndColor(Channel));
 	}
+	return Colors;
 }
 
 
@@ -1554,7 +1469,7 @@ bool OpenColorPicker(const FColorPickerArgs& Args)
 	TSharedPtr<SWindow> Window = nullptr;
 	TSharedRef<SBorder> WindowContent = SNew(SBorder)
 			.BorderImage(FAppStyle::Get().GetBrush("Brushes.Panel"))
-			.Padding(FMargin(8.0f, 8.0f));
+			.Padding(FMargin(16.0f, 16.0f));
 	
 	bool bNeedToAddWindow = true;
 	if (!bOverrideNonModalCreation)
@@ -1592,7 +1507,6 @@ bool OpenColorPicker(const FColorPickerArgs& Args)
 	TSharedRef<SColorPicker> CreatedColorPicker = SNew(SColorPicker)
 		.TargetColorAttribute(OldColor)
 		.UseAlpha(Args.bUseAlpha)
-		.ExpandAdvancedSection(Args.bExpandAdvancedSection)
 		.OnlyRefreshOnMouseUp(Args.bOnlyRefreshOnMouseUp && !Args.bIsModal)
 		.OnlyRefreshOnOk(Args.bOnlyRefreshOnOk || Args.bIsModal)
 		.OnColorCommitted(Args.OnColorCommitted)
diff --git a/Engine/Source/Runtime/AppFramework/Private/Widgets/Colors/SColorSlider.cpp b/Engine/Source/Runtime/AppFramework/Private/Widgets/Colors/SColorSlider.cpp
new file mode 100644
--- /dev/null
+++ b/Engine/Source/Runtime/AppFramework/Private/Widgets/Colors/SColorSlider.cpp
@@ -0,0 +1,251 @@
+// Copyright Epic Games, Inc. All Rights Reserved.
+
+#include "SColorSlider.h"
+
+#include "Widgets/Text/STextBlock.h"
+
+SColorSlider::SColorSlider()
+	: Orientation(*this, EOrientation::Orient_Horizontal)
+	, GradientColors(*this)
+	, bHasAlphaBackground(*this, false)
+{
+}
+
+void SColorSlider::Construct(const FArguments& InArgs)
+{
+	GradientColors.Assign(*this, InArgs._GradientColors);
+
+	bHasAlphaBackground.Assign(*this, InArgs._HasAlphaBackground.Get());
+	Orientation.Assign(*this, InArgs._Orientation.Get());
+
+	BorderBrush = FAppStyle::Get().GetBrush("ColorPicker.RoundedInputBorder");
+	BorderActiveBrush = FAppStyle::Get().GetBrush("ColorPicker.RoundedInputBorderActive");
+	BorderHoveredBrush = FAppStyle::Get().GetBrush("ColorPicker.RoundedInputBorderHovered");
+	AlphaBackgroundBrush = FAppStyle::Get().GetBrush("ColorPicker.RoundedAlphaBackground");
+
+	const FSlateFontInfo SmallFont = FAppStyle::Get().GetFontStyle("ColorPicker.SmallFont");
+
+	ColorSliderSize = (Orientation.Get() == Orient_Horizontal) ? HorizontalSliderLength : VerticalSliderHeight;
+
+	SAssignNew(Slider, SSlider)
+		.IndentHandle(false)
+		.Orientation(Orientation.Get())
+		.SliderBarColor(FLinearColor::Transparent)
+		.SliderHandleColor(FLinearColor::Transparent)
+		.Style(&FAppStyle::Get().GetWidgetStyle<FSliderStyle>("ColorPicker.Slider"))
+		.MinValue(InArgs._MinSliderValue.Get())
+		.MaxValue(InArgs._MaxSliderValue.Get())
+		.StepSize(InArgs._Delta.Get())
+		.Value(InArgs._Value)
+		.OnMouseCaptureBegin(InArgs._OnBeginSliderMovement)
+		.OnMouseCaptureEnd(InArgs._OnEndSliderMovement)
+		.OnValueChanged(InArgs._OnValueChanged);
+
+	TSharedPtr<SHorizontalBox> ColorWidget = SNew(SHorizontalBox);
+
+	if (Orientation.Get() == Orient_Horizontal)
+	{
+		ColorWidget->AddSlot()
+			.MinWidth(LabelSize)
+			.MaxWidth(LabelSize)
+			.Padding(0.0f, 0.0f, Padding, 0.0f)
+			[
+				SNew(SVerticalBox)
+
+				+ SVerticalBox::Slot()
+					.VAlign(VAlign_Center)
+					[
+						SNew(STextBlock)
+							.Text(InArgs._Label)
+							.TextStyle(&FAppStyle::Get().GetWidgetStyle<FTextBlockStyle>("SmallText"))
+							.Justification(ETextJustify::Left)
+					]
+			];
+
+		ColorWidget->AddSlot()
+			.MinWidth(ColorSliderSize)
+			.MaxWidth(ColorSliderSize)
+			[
+				Slider.ToSharedRef()
+			];
+
+		ColorWidget->AddSlot()
+			.MinWidth(SpinBoxSize)
+			.MaxWidth(SpinBoxSize)
+			.Padding(Padding, 0.0f, 0.0f, 0.0f)
+			[
+				SNew(SSpinBox<float>)
+					.MinValue(InArgs._MinSpinBoxValue.Get())
+					.MaxValue(InArgs._MaxSpinBoxValue.Get())
+					.MinSliderValue(InArgs._MinSliderValue.Get())
+					.MaxSliderValue(InArgs._MaxSliderValue.Get())
+					.MaxFractionalDigits(3)
+					.Delta(InArgs._Delta)
+					.Value(InArgs._Value)
+					.OnBeginSliderMovement(InArgs._OnBeginSpinBoxMovement)
+					.OnEndSliderMovement(InArgs._OnEndSpinBoxMovement)
+					.OnValueChanged(InArgs._OnValueChanged)
+					.OnValueCommitted(this, &SColorSlider::OnSpinBoxValueCommitted)
+					.Font(SmallFont)
+			];
+	}
+	else
+	{
+		ColorWidget->AddSlot()
+			[
+				Slider.ToSharedRef()
+			];
+	}
+
+	ChildSlot
+	[
+		ColorWidget.ToSharedRef()
+	];
+}
+
+FVector2D SColorSlider::ComputeDesiredSize(float) const
+{
+	if (Orientation.Get() == EOrientation::Orient_Horizontal)
+	{
+		constexpr float TotalWidth = LabelSize + Padding + HorizontalSliderLength + Padding + SpinBoxSize;
+		return FVector2D(TotalWidth, HorizontalSliderHeight);
+	}
+	else
+	{
+		return FVector2D(VerticalSliderWidth, VerticalSliderHeight);
+	}
+}
+
+int32 SColorSlider::OnPaint(const FPaintArgs& Args, const FGeometry& AllottedGeometry, const FSlateRect& MyCullingRect, FSlateWindowElementList& OutDrawElements, int32 LayerId, const FWidgetStyle& InWidgetStyle, bool bParentEnabled) const
+{
+	// Render the underlying widgets underneath so we can draw the gradient, selector, and border on top
+	const int32 MaxChildLayer = SCompoundWidget::OnPaint(Args, AllottedGeometry, MyCullingRect, OutDrawElements, LayerId, InWidgetStyle, bParentEnabled);
+	LayerId = MaxChildLayer + 1;
+
+	const ESlateDrawEffect DrawEffects = ESlateDrawEffect::None;
+
+	FVector2f SliderSize;
+	FVector2f SliderOffset = FVector2f(0.0f, 0.0f);
+
+	if (Orientation.Get() == Orient_Horizontal)
+	{
+		SliderSize = FVector2f(ColorSliderSize, AllottedGeometry.GetLocalSize().Y);
+		SliderOffset.X = LabelSize + Padding;
+	}
+	else
+	{
+		SliderSize = FVector2f(AllottedGeometry.GetLocalSize().X, ColorSliderSize);
+	}
+
+	// Draw Color Gradient
+	const int32 NumColors = GradientColors.Get().Num();
+
+	if (NumColors > 0)
+	{
+		if (bHasAlphaBackground.Get())
+		{
+			FSlateDrawElement::MakeBox(
+				OutDrawElements,
+				LayerId++,
+				AllottedGeometry.ToPaintGeometry(SliderSize, FSlateLayoutTransform(SliderOffset)),
+				AlphaBackgroundBrush,
+				DrawEffects
+			);
+		}
+
+		// Vertical sliders need the start color at the bottom and the end color at the bottom
+		if (Orientation.Get() == Orient_Vertical)
+		{
+			Algo::Reverse(GradientColors.Get());
+		}
+
+		// Generate the gradient stops for the colors
+		TArray<FSlateGradientStop> GradientStops;
+		for (int32 ColorIndex = 0; ColorIndex < NumColors; ++ColorIndex)
+		{
+			GradientStops.Add(FSlateGradientStop(ColorSliderSize * (float(ColorIndex) / (NumColors - 1)), GradientColors.Get()[ColorIndex]));
+		}
+
+		const FVector4f CornerRadius = FVector4f(4.0f, 4.0f, 4.0f, 4.0f);
+
+		FSlateDrawElement::MakeGradient(
+			OutDrawElements,
+			LayerId++,
+			AllottedGeometry.ToPaintGeometry(SliderSize, FSlateLayoutTransform(SliderOffset)),
+			GradientStops,
+			(Orientation.Get() == Orient_Horizontal) ? EOrientation::Orient_Vertical : EOrientation::Orient_Horizontal,
+			DrawEffects | ESlateDrawEffect::NoGamma, // Gradient colors are in linear space, so disable gamma to let them blend properly
+			CornerRadius
+		);	
+	}
+
+	// Draw Selector
+	const FSlateBrush* SelectorBrush;
+	const FVector2f AllottedGeometrySize = AllottedGeometry.GetLocalSize();
+
+	// If the current slider value is greater than the slider's maximum value, update the slider maximum.
+	// This could occur because the maximum value of the spinbox might be greater than the maximum value of the slider. 
+	const float SliderValue = Slider->GetValue();
+	if (SliderValue > Slider->GetMaxValue())
+	{
+		Slider->SetMinAndMaxValues(Slider->GetMinValue(), SliderValue);
+	}
+
+	float FractionFilled = SliderValue / Slider->GetMaxValue();
+
+	FVector2f SelectorSize;
+	FVector2f SelectorOffset;
+
+	constexpr float SelectorThickness = 3.0f;
+
+	if (Orientation.Get() == Orient_Horizontal)
+	{
+		SelectorBrush = FAppStyle::Get().GetBrush("ColorPicker.SpinBoxSelectorVertical");
+		SelectorSize = FVector2f(SelectorThickness, HorizontalSliderHeight - 2.0f);
+
+		float SelectorRange = ColorSliderSize - SelectorSize.X;
+		SelectorOffset = FVector2f(SelectorRange * FractionFilled, 1.0f);
+	}
+	else
+	{
+		SelectorBrush = FAppStyle::Get().GetBrush("ColorPicker.SpinBoxSelectorHorizontal");
+		SelectorSize = FVector2f(VerticalSliderWidth - 2.0f, SelectorThickness);
+
+		// Invert so that 1 is the top of the slider and 0 is the bottom of the slider
+		FractionFilled = 1.0f - FractionFilled;
+
+		float SelectorRange = ColorSliderSize - SelectorSize.Y;
+		SelectorOffset = FVector2f(1.0f, SelectorRange * FractionFilled);
+	}
+
+	FSlateDrawElement::MakeBox(
+		OutDrawElements,
+		LayerId++,
+		AllottedGeometry.ToPaintGeometry(SelectorSize, FSlateLayoutTransform(SelectorOffset + SliderOffset)),
+		SelectorBrush,
+		DrawEffects,
+		SelectorBrush->GetTint(InWidgetStyle) * InWidgetStyle.GetColorAndOpacityTint()
+	);
+
+	// Draw Border
+	const FSlateBrush* BorderImage = Slider->HasMouseCapture() ? BorderActiveBrush : (Slider->IsHovered()) ? BorderHoveredBrush : BorderBrush;
+
+	FSlateDrawElement::MakeBox(
+		OutDrawElements,
+		LayerId++,
+		AllottedGeometry.ToPaintGeometry(SliderSize, FSlateLayoutTransform(SliderOffset)),
+		BorderImage,
+		DrawEffects,
+		BorderImage->GetTint(InWidgetStyle) * InWidgetStyle.GetColorAndOpacityTint()
+	);
+
+	return LayerId;
+}
+
+void SColorSlider::OnSpinBoxValueCommitted(float NewValue, ETextCommit::Type CommitType)
+{
+	if (NewValue > Slider->GetMaxValue())
+	{
+		Slider->SetMinAndMaxValues(Slider->GetMinValue(), NewValue);
+	}
+}
diff --git a/Engine/Source/Runtime/AppFramework/Private/Widgets/Colors/SColorSlider.h b/Engine/Source/Runtime/AppFramework/Private/Widgets/Colors/SColorSlider.h
new file mode 100644
--- /dev/null
+++ b/Engine/Source/Runtime/AppFramework/Private/Widgets/Colors/SColorSlider.h
@@ -0,0 +1,138 @@
+// Copyright Epic Games, Inc. All Rights Reserved.
+
+#pragma once
+
+#include "Widgets/SCompoundWidget.h"
+
+#include "Styling/SlateTypes.h"
+#include "Widgets/DeclarativeSyntaxSupport.h"
+#include "Widgets/Input/SSlider.h"
+#include "Widgets/Input/SSpinBox.h"
+
+/**
+ * Implements a color slider widget. 
+ * When the orientation is horizontal, it features an STextBlock label, a SSlider with color gradient drawn on top, and a SSpinBox.
+ * When the orientation is vertical, it features only the SSlider with color gradient
+ */
+class SColorSlider : public SCompoundWidget
+{
+public:
+	SColorSlider();
+
+	/** Notification for numeric value change */
+	DECLARE_DELEGATE_OneParam(FOnValueChanged, float);
+
+	SLATE_BEGIN_ARGS(SColorSlider)
+		: _Value(0.0f)
+		, _MinSpinBoxValue(0.0f)
+		, _MaxSpinBoxValue(1.0f)
+		, _MinSliderValue(0.0f)
+		, _MaxSliderValue(1.0f)
+		, _Delta(0.01f)
+		, _Orientation(Orient_Horizontal)
+		, _Label()
+		, _GradientColors()
+		, _HasAlphaBackground(false)
+		, _OnValueChanged()
+		, _OnBeginSliderMovement()
+		, _OnEndSliderMovement()
+		, _OnBeginSpinBoxMovement()
+		, _OnEndSpinBoxMovement()
+	{ }
+
+		/** The value that determines where the slider handle is drawn */
+		SLATE_ATTRIBUTE(float, Value)
+
+		/** The minimum value of the spinbox */
+		SLATE_ATTRIBUTE(float, MinSpinBoxValue)
+
+		/** The maximum value of the spinbox */
+		SLATE_ATTRIBUTE(float, MaxSpinBoxValue)
+
+		/** The minimum value of the slider */
+		SLATE_ATTRIBUTE(float, MinSliderValue)
+
+		/** The maximum value of the slider */
+		SLATE_ATTRIBUTE(float, MaxSliderValue)
+
+		/** The delta to increment the value as the slider moves */
+		SLATE_ATTRIBUTE(float, Delta)
+
+		/** Orientation of the slider */
+		SLATE_ATTRIBUTE(EOrientation, Orientation)
+
+		/** Text content of the TextBlock (horizontal sliders only) */
+		SLATE_ATTRIBUTE(FText, Label)
+
+		/** List of colors which determine the gradient stops drawn on top of the slider */
+		SLATE_ATTRIBUTE(TArray<FLinearColor>, GradientColors)
+
+		/** Whether a checker background is displayed for alpha viewing */
+		SLATE_ATTRIBUTE(bool, HasAlphaBackground)
+
+		/** Called when the value is changed by the Slider or SpinBox */
+		SLATE_EVENT(FOnValueChanged, OnValueChanged)
+
+		/** Called right before the slider handle on the Slider widget begins to move */
+		SLATE_EVENT(FSimpleDelegate, OnBeginSliderMovement)
+
+		/** Called right after the slider handle on the Slider widget is released by the user */
+		SLATE_EVENT(FSimpleDelegate, OnEndSliderMovement)
+
+		/** Called right before the slider handle on the SpinBox widget begins to move */
+		SLATE_EVENT(FSimpleDelegate, OnBeginSpinBoxMovement)
+
+		/** Called right after the slider handle on the SpinBox widget is released by the user */
+		SLATE_EVENT(FOnValueChanged, OnEndSpinBoxMovement)
+
+	SLATE_END_ARGS()
+
+	/**
+	 * Construct this widget.
+	 *
+	 * @param InArgs The declaration data for this widget.
+	 */
+	void Construct(const FArguments& InArgs);
+
+protected:
+	// Begin SWidget overrides
+	virtual int32 OnPaint(const FPaintArgs& Args, const FGeometry& AllottedGeometry, const FSlateRect& MyCullingRect, FSlateWindowElementList& OutDrawElements, int32 LayerId, const FWidgetStyle& InWidgetStyle, bool bParentEnabled) const override;
+	virtual FVector2D ComputeDesiredSize(float) const override;
+	// End SWidget overrides
+
+private:
+	/** Update the max value of the slider if the value committed is higher than the current maximum slider value */
+	void OnSpinBoxValueCommitted(float NewValue, ETextCommit::Type CommitType);
+
+private:
+	/** Orientation of the slider */
+	TSlateAttribute<EOrientation, EInvalidateWidgetReason::Paint> Orientation;
+
+	/** The colors used in the gradient. */
+	TSlateAttribute<TArray<FLinearColor>, EInvalidateWidgetReason::Paint> GradientColors;
+
+	/** Whether a checker background is displayed for alpha viewing */
+	TSlateAttribute<bool, EInvalidateWidgetReason::Paint> bHasAlphaBackground;
+
+	/** Slider widget */
+	TSharedPtr<SSlider> Slider;
+
+	/** Widget layout sizes */
+	float ColorSliderSize = 0.0f;
+
+	static constexpr float Padding = 8.0f;
+	static constexpr float LabelSize = 8.0f;
+	static constexpr float SpinBoxSize = 60.0f;
+
+	static constexpr float HorizontalSliderLength = 123.0f;
+	static constexpr float HorizontalSliderHeight = 20.0f;
+
+	static constexpr float VerticalSliderWidth = 28.0f;
+	static constexpr float VerticalSliderHeight = 200.0f;
+
+	/** Styling brushes */
+	const FSlateBrush* BorderBrush;
+	const FSlateBrush* BorderActiveBrush;
+	const FSlateBrush* BorderHoveredBrush;
+	const FSlateBrush* AlphaBackgroundBrush;
+};
diff --git a/Engine/Source/Runtime/AppFramework/Private/Widgets/Colors/SEyeDropperButton.cpp b/Engine/Source/Runtime/AppFramework/Private/Widgets/Colors/SEyeDropperButton.cpp
--- a/Engine/Source/Runtime/AppFramework/Private/Widgets/Colors/SEyeDropperButton.cpp
+++ b/Engine/Source/Runtime/AppFramework/Private/Widgets/Colors/SEyeDropperButton.cpp
@@ -20,13 +20,12 @@ void SEyeDropperButton::Construct(const FArguments& InArgs)
 	// Their visibility and colors are changed according to whether dropper mode is active or not.
 	SButton::Construct(
 		SButton::FArguments()
-		.ContentPadding(1.0f)
+		.ContentPadding(FMargin(2.0f, 2.5f))
 		.OnClicked(this, &SEyeDropperButton::OnClicked)
 		[
 			SNew(SOverlay)
 
 			+ SOverlay::Slot()
-			.Padding(FMargin(1.0f, 0.0f))
 			[
 				SNew(SImage)
 				.Image(FCoreStyle::Get().GetBrush("ColorPicker.EyeDropper"))
@@ -36,10 +35,16 @@ void SEyeDropperButton::Construct(const FArguments& InArgs)
 
 			+ SOverlay::Slot()
 			[
-				SNew(STextBlock)
-				.Text(LOCTEXT("EscapeCue", "Esc"))
-				.ToolTipText(LOCTEXT("EyeDropperEscapeCue_ToolTip", "Hit Escape key to stop the eye dropper"))
-				.Visibility(this, &SEyeDropperButton::GetEscapeTextVisibility)
+				SNew(SHorizontalBox)
+
+				+ SHorizontalBox::Slot()
+				.VAlign(VAlign_Center)
+				[
+					SNew(STextBlock)
+					.Text(LOCTEXT("EscapeCue", "Esc"))
+					.ToolTipText(LOCTEXT("EyeDropperEscapeCue_ToolTip", "Hit Escape key to stop the eye dropper"))
+					.Visibility(this, &SEyeDropperButton::GetEscapeTextVisibility)
+				]
 			]
 		]
 	);
diff --git a/Engine/Source/Runtime/AppFramework/Public/Widgets/Colors/SColorPicker.h b/Engine/Source/Runtime/AppFramework/Public/Widgets/Colors/SColorPicker.h
--- a/Engine/Source/Runtime/AppFramework/Public/Widgets/Colors/SColorPicker.h
+++ b/Engine/Source/Runtime/AppFramework/Public/Widgets/Colors/SColorPicker.h
@@ -267,8 +267,8 @@ private:
 	// Callback for getting the start color of a color spin box gradient.
 	APPFRAMEWORK_API FLinearColor GetGradientStartColor(EColorPickerChannels Channel) const;
 
-	// Callback for handling expansion of the 'Advanced' area.
-	APPFRAMEWORK_API void HandleAdvancedAreaExpansionChanged(bool Expanded);
+	// Callback for getting the array of colors for a color spin box gradient.
+	APPFRAMEWORK_API TArray<FLinearColor> GetGradientColors(EColorPickerChannels Channel) const;
 
 	// Callback for getting the visibility of the alpha channel portion in color blocks.
 	APPFRAMEWORK_API EVisibility HandleAlphaColorBlockVisibility() const;
diff --git a/Engine/Source/Runtime/Slate/Private/Widgets/Colors/SColorWheel.cpp b/Engine/Source/Runtime/Slate/Private/Widgets/Colors/SColorWheel.cpp
--- a/Engine/Source/Runtime/Slate/Private/Widgets/Colors/SColorWheel.cpp
+++ b/Engine/Source/Runtime/Slate/Private/Widgets/Colors/SColorWheel.cpp
@@ -70,7 +70,7 @@ FReply SColorWheel::OnMouseButtonUp(const FGeometry& MyGeometry, const FPointerE
 		OnMouseCaptureEnd.ExecuteIfBound();
 
 		// Before showing the mouse position again, reset its position to the final location of the selector on the color wheel
-		const FVector2f CircleSize = MyGeometry.GetLocalSize() - SelectorImage->ImageSize;
+		const FVector2f CircleSize = MyGeometry.GetLocalSize();
 		const FVector2f FinalMousePosition = 0.5f * (MyGeometry.GetLocalSize() + CalcRelativePositionFromCenter() * CircleSize);
 
 		return FReply::Handled().ReleaseMouseCapture().SetMousePos(MyGeometry.LocalToAbsolute(FinalMousePosition).IntPoint());
@@ -104,21 +104,25 @@ int32 SColorWheel::OnPaint(const FPaintArgs& Args, const FGeometry& AllottedGeom
 	const bool bIsEnabled = ShouldBeEnabled(bParentEnabled);
 	const ESlateDrawEffect DrawEffects = bIsEnabled ? ESlateDrawEffect::None : ESlateDrawEffect::DisabledEffect;
 	const FVector2f SelectorSize = SelectorImage->ImageSize;
-	const FVector2f CircleSize = AllottedGeometry.GetLocalSize() - SelectorSize;
+	const FVector2f CircleCenter = AllottedGeometry.GetLocalSize() * 0.5f;
 	
+	// Draw Color Wheel Image
 	FSlateDrawElement::MakeBox(
 		OutDrawElements,
 		LayerId,
-		AllottedGeometry.ToPaintGeometry(CircleSize, FSlateLayoutTransform(0.5f * SelectorSize)),
+		AllottedGeometry.ToPaintGeometry(),
 		Image,
 		DrawEffects,
 		InWidgetStyle.GetColorAndOpacityTint() * Image->GetTint(InWidgetStyle)
 	);
 
+	// Draw Selector Image
+	const FVector2f SelectorImageLocation = CircleCenter + (CalcRelativePositionFromCenter() * CircleCenter) - (0.5f * SelectorSize);
+
 	FSlateDrawElement::MakeBox(
 		OutDrawElements,
 		LayerId + 1,
-		AllottedGeometry.ToPaintGeometry(SelectorSize, FSlateLayoutTransform(0.5f * (AllottedGeometry.GetLocalSize() + CalcRelativePositionFromCenter() * CircleSize - SelectorSize))),
+		AllottedGeometry.ToPaintGeometry(SelectorSize, FSlateLayoutTransform(SelectorImageLocation)),
 		SelectorImage,
 		DrawEffects,
 		InWidgetStyle.GetColorAndOpacityTint() * SelectorImage->GetTint(InWidgetStyle)
@@ -161,7 +165,7 @@ bool SColorWheel::ProcessMouseAction(const FGeometry& MyGeometry, const FPointer
 		LocalMouseCoordinate = MyGeometry.AbsoluteToLocal(MouseEvent.GetScreenSpacePosition());
 	}
 	
-	const FVector2f RelativePositionFromCenter = (2.0f * LocalMouseCoordinate - MyGeometry.GetLocalSize()) / (MyGeometry.GetLocalSize() - SelectorImage->ImageSize);
+	const FVector2f RelativePositionFromCenter = (2.0f * LocalMouseCoordinate - MyGeometry.GetLocalSize()) / (MyGeometry.GetLocalSize());
 	const float RelativeRadius = RelativePositionFromCenter.Size();
 
 	if (RelativeRadius <= 1.0f || bProcessWhenOutsideColorWheel)
diff --git a/Engine/Source/Runtime/Slate/Public/Widgets/Colors/SColorBlock.h b/Engine/Source/Runtime/Slate/Public/Widgets/Colors/SColorBlock.h
--- a/Engine/Source/Runtime/Slate/Public/Widgets/Colors/SColorBlock.h
+++ b/Engine/Source/Runtime/Slate/Public/Widgets/Colors/SColorBlock.h
@@ -30,7 +30,7 @@ public:
 
 	SLATE_BEGIN_ARGS(SColorBlock)
 		: _Color(FLinearColor::White)
-		, _AlphaBackgroundBrush(FAppStyle::Get().GetBrush("ColorPicker.AlphaBackground"))
+		, _AlphaBackgroundBrush(FAppStyle::Get().GetBrush("ColorPicker.RoundedAlphaBackground"))
 		, _CornerRadius(0.0f)
 		, _ColorIsHSV(false)
 		, _ShowBackgroundForAlpha(false)
diff --git a/Engine/Source/Runtime/SlateCore/Private/Styling/StarshipCoreStyle.cpp b/Engine/Source/Runtime/SlateCore/Private/Styling/StarshipCoreStyle.cpp
--- a/Engine/Source/Runtime/SlateCore/Private/Styling/StarshipCoreStyle.cpp
+++ b/Engine/Source/Runtime/SlateCore/Private/Styling/StarshipCoreStyle.cpp
@@ -1945,13 +1945,18 @@ void FStarshipCoreStyle::SetupColorPickerStyles(TSharedRef<FStyle>& Style)
 	// SColorPicker defaults...
 	{
 		Style->Set("ColorPicker.RoundedSolidBackground", new FSlateRoundedBoxBrush(FStyleColors::White, InputFocusRadius));
-		Style->Set("ColorPicker.RoundedAlphaBackground", new FSlateRoundedBoxBrush(FName(*RootToContentDir("Starship/Common/Checker", TEXT(".png"))), FLinearColor::White, InputFocusRadius, FLinearColor::White, 0.0f, Icon16x16, ESlateBrushTileType::Both));
+		Style->Set("ColorPicker.RoundedAlphaBackground", new FSlateRoundedBoxBrush(FName(*RootToContentDir("Starship/Common/Checker", TEXT(".png"))), FLinearColor::Transparent, InputFocusRadius, FLinearColor::Transparent, 0.0f, Icon16x16, ESlateBrushTileType::Both));
 		Style->Set("ColorPicker.RoundedInputBorder", new FSlateRoundedBoxBrush(FStyleColors::Transparent, InputFocusRadius, FStyleColors::InputOutline, InputFocusThickness));
+		Style->Set("ColorPicker.RoundedInputBorderActive", new FSlateRoundedBoxBrush(FStyleColors::Transparent, InputFocusRadius, FStyleColors::Primary, InputFocusThickness));
+		Style->Set("ColorPicker.RoundedInputBorderHovered", new FSlateRoundedBoxBrush(FStyleColors::Transparent, InputFocusRadius, FStyleColors::Hover, InputFocusThickness));
 		Style->Set("ColorPicker.MultipleValuesBackground", new FSlateRoundedBoxBrush(FStyleColors::Input, InputFocusRadius));
 		Style->Set("ColorPicker.AlphaBackground", new IMAGE_BRUSH("Starship/Common/Checker", Icon16x16, FLinearColor::White, ESlateBrushTileType::Both));
 		Style->Set("ColorPicker.EyeDropper", new IMAGE_BRUSH("Icons/eyedropper_16px", Icon16x16));
 		Style->Set("ColorPicker.Font", FStyleFonts::Get().Normal);
+		Style->Set("ColorPicker.SmallFont", FStyleFonts::Get().Small);
 		Style->Set("ColorPicker.Mode", new IMAGE_BRUSH("Common/ColorPicker_Mode_16x", Icon16x16));
+		Style->Set("ColorPicker.ModeWheel", new IMAGE_BRUSH("Common/ColorWheel_20x", Icon20x20));
+		Style->Set("ColorPicker.ModeSpectrum", new IMAGE_BRUSH("Common/ColorSpectrum_20x", Icon20x20));
 		Style->Set("ColorPicker.Separator", new IMAGE_BRUSH("Common/ColorPicker_Separator", FVector2f(2.0f, 2.0f)));
 		Style->Set("ColorPicker.Selector", new IMAGE_BRUSH("Common/Circle", Icon8x8));
 		Style->Set("ColorPicker.Slider", FSliderStyle()
@@ -1959,6 +1964,8 @@ void FStarshipCoreStyle::SetupColorPickerStyles(TSharedRef<FStyle>& Style)
 			.SetNormalThumbImage(IMAGE_BRUSH("Common/ColorPicker_SliderHandle", FVector2f(8.0f, 32.0f)))
 			.SetHoveredThumbImage(IMAGE_BRUSH("Common/ColorPicker_SliderHandle", FVector2f(8.0f, 32.0f)))
 		);
+		Style->Set("ColorPicker.SpinBoxSelectorVertical", new IMAGE_BRUSH("Starship/Common/color-grading-spinbox-selector", FVector2f(3.0f, 1.0f)));
+		Style->Set("ColorPicker.SpinBoxSelectorHorizontal", new IMAGE_BRUSH("Starship/Common/color-grading-spinbox-selector-v", FVector2f(1.0f, 3.0f)));
 	}
 
 	// SColorSpectrum defaults...
