ColorPicker: Change the look of the Saturation and Value sliders.
Fix spinboxes to retain values greater than 1.0f when spinning.
Remove fill color from the color slider's numeric spinbox.
#rnx

diff --git a/Engine/Source/Runtime/AppFramework/Private/Widgets/Colors/SColorPicker.cpp b/Engine/Source/Runtime/AppFramework/Private/Widgets/Colors/SColorPicker.cpp
--- a/Engine/Source/Runtime/AppFramework/Private/Widgets/Colors/SColorPicker.cpp
+++ b/Engine/Source/Runtime/AppFramework/Private/Widgets/Colors/SColorPicker.cpp
@@ -960,8 +960,8 @@ FLinearColor SColorPicker::GetGradientEndColor( EColorPickerChannels Channel ) c
 	case EColorPickerChannels::Green:		return FLinearColor::Green;
 	case EColorPickerChannels::Blue:		return FLinearColor::Blue;
 	case EColorPickerChannels::Alpha:		return FLinearColor::White;
-	case EColorPickerChannels::Saturation:	return FLinearColor(CurrentColorHSV.R, 1.0f, CurrentColorHSV.B, 1.0f).HSVToLinearRGB();
-	case EColorPickerChannels::Value:		return FLinearColor::White;
+	case EColorPickerChannels::Saturation:	return FLinearColor(CurrentColorHSV.R, 1.0f, 1.0f, 1.0f).HSVToLinearRGB();
+	case EColorPickerChannels::Value:		return FLinearColor(CurrentColorHSV.R, CurrentColorHSV.G, 1.0f, 1.0f).HSVToLinearRGB();
 	default:								return FLinearColor();
 	}
 }
@@ -975,8 +975,8 @@ FLinearColor SColorPicker::GetGradientStartColor( EColorPickerChannels Channel )
 	case EColorPickerChannels::Green:		return FLinearColor::Black;
 	case EColorPickerChannels::Blue:		return FLinearColor::Black;
 	case EColorPickerChannels::Alpha:		return FLinearColor::Transparent;
-	case EColorPickerChannels::Saturation:	return FLinearColor(CurrentColorHSV.R, 0.0f, CurrentColorHSV.B, 1.0f).HSVToLinearRGB();
-	case EColorPickerChannels::Value:		return FLinearColor::Black;
+	case EColorPickerChannels::Saturation:	return FLinearColor(CurrentColorHSV.R, 0.0f, 1.0f, 1.0f).HSVToLinearRGB();
+	case EColorPickerChannels::Value:		return FLinearColor(CurrentColorHSV.R, CurrentColorHSV.G, 0.0f, 1.0f).HSVToLinearRGB();
 	default:								return FLinearColor();
 	}
 }
diff --git a/Engine/Source/Runtime/AppFramework/Private/Widgets/Colors/SColorSlider.cpp b/Engine/Source/Runtime/AppFramework/Private/Widgets/Colors/SColorSlider.cpp
--- a/Engine/Source/Runtime/AppFramework/Private/Widgets/Colors/SColorSlider.cpp
+++ b/Engine/Source/Runtime/AppFramework/Private/Widgets/Colors/SColorSlider.cpp
@@ -75,6 +75,7 @@ void SColorSlider::Construct(const FArguments& InArgs)
 			.Padding(Padding, 0.0f, 0.0f, 0.0f)
 			[
 				SNew(SSpinBox<float>)
+					.Style(FAppStyle::Get(), "ColorSlider.SpinBox")
 					.MinValue(InArgs._MinSpinBoxValue.Get())
 					.MaxValue(InArgs._MaxSpinBoxValue.Get())
 					.MinSliderValue(InArgs._MinSliderValue.Get())
@@ -82,6 +83,7 @@ void SColorSlider::Construct(const FArguments& InArgs)
 					.MaxFractionalDigits(3)
 					.Delta(InArgs._Delta)
 					.Value(InArgs._Value)
+					.SupportDynamicSliderMaxValue(true)
 					.OnBeginSliderMovement(InArgs._OnBeginSpinBoxMovement)
 					.OnEndSliderMovement(InArgs._OnEndSpinBoxMovement)
 					.OnValueChanged(InArgs._OnValueChanged)
diff --git a/Engine/Source/Runtime/SlateCore/Private/Styling/StarshipCoreStyle.cpp b/Engine/Source/Runtime/SlateCore/Private/Styling/StarshipCoreStyle.cpp
--- a/Engine/Source/Runtime/SlateCore/Private/Styling/StarshipCoreStyle.cpp
+++ b/Engine/Source/Runtime/SlateCore/Private/Styling/StarshipCoreStyle.cpp
@@ -1968,6 +1968,16 @@ void FStarshipCoreStyle::SetupColorPickerStyles(TSharedRef<FStyle>& Style)
 		Style->Set("ColorPicker.SpinBoxSelectorHorizontal", new IMAGE_BRUSH("Starship/Common/color-grading-spinbox-selector-v", FVector2f(1.0f, 3.0f)));
 	}
 
+	// SColorSlider defaults...
+	{
+		FSpinBoxStyle ColorSliderSpinBox = Style->GetWidgetStyle<FSpinBoxStyle>("SpinBox");
+		ColorSliderSpinBox.SetActiveFillBrush(FSlateRoundedBoxBrush(FStyleColors::Input, InputFocusRadius, FLinearColor::Transparent, InputFocusThickness));
+		ColorSliderSpinBox.SetHoveredFillBrush(FSlateRoundedBoxBrush(FStyleColors::Input, InputFocusRadius, FLinearColor::Transparent, InputFocusThickness));
+		ColorSliderSpinBox.SetInactiveFillBrush(FSlateRoundedBoxBrush(FStyleColors::Input, InputFocusRadius, FLinearColor::Transparent, InputFocusThickness));
+
+		Style->Set("ColorSlider.SpinBox", ColorSliderSpinBox);
+	}
+
 	// SColorSpectrum defaults...
 	{
 		Style->Set("ColorSpectrum.Spectrum", new IMAGE_BRUSH("Common/ColorSpectrum", FVector2f(256.f, 256.f)));
