SColorPicker: Implement new Color Themes panel UI based on new design specifications.
Update panel layout, using uniform block sizes, and auto-resize to add new rows when needed.
Update how placeholder color blocks are displayed.
Add "Recents" color theme, which keeps the most recently selected colors. This is not editable by the user like other color themes.
Add dropdown menu button that opens submenu for Recents, Saved Themes, New Theme, Rename, Duplicate, Delete.
Add overlay button to color preview blocks to add the old/current color to the current theme.

diff --git a/Engine/Source/Runtime/AppFramework/AppFramework.Build.cs b/Engine/Source/Runtime/AppFramework/AppFramework.Build.cs
--- a/Engine/Source/Runtime/AppFramework/AppFramework.Build.cs
+++ b/Engine/Source/Runtime/AppFramework/AppFramework.Build.cs
@@ -21,6 +21,7 @@ public class AppFramework : ModuleRules
 		PrivateDependencyModuleNames.AddRange(
 			new string[] {
 				"InputCore",
+				"ToolMenus",
 			}
 		);
 
diff --git a/Engine/Source/Runtime/AppFramework/Private/Widgets/Colors/SColorPicker.cpp b/Engine/Source/Runtime/AppFramework/Private/Widgets/Colors/SColorPicker.cpp
--- a/Engine/Source/Runtime/AppFramework/Private/Widgets/Colors/SColorPicker.cpp
+++ b/Engine/Source/Runtime/AppFramework/Private/Widgets/Colors/SColorPicker.cpp
@@ -38,20 +38,12 @@ const FVector2D SColorPicker::DEFAULT_WINDOW_SIZE = FVector2D(462, 526);
 /** The max time allowed for updating before we shut off auto-updating */
 const double SColorPicker::MAX_ALLOWED_UPDATE_TIME = 0.1;
 
-TWeakPtr<SColorThemesViewer> SColorPicker::ColorThemesViewer;
-
 
 /* SColorPicker structors
  *****************************************************************************/
 
 SColorPicker::~SColorPicker()
 {
-	TSharedPtr<SColorThemesViewer> ThemesViewer = ColorThemesViewer.Pin();
-
-	if (ThemesViewer.IsValid())
-	{
-		ThemesViewer->OnCurrentThemeChanged().RemoveAll(this);
-	}
 }
 
 
@@ -79,12 +71,6 @@ void SColorPicker::Construct( const FArguments& InArgs )
 	bClampValue = InArgs._ClampValue;
 	OptionalOwningDetailsView = InArgs._OptionalOwningDetailsView.Get().IsValid() ? InArgs._OptionalOwningDetailsView.Get() : nullptr;
 
-	if ( InArgs._sRGBOverride.IsSet() )
-	{
-		OriginalSRGBOption = SColorThemesViewer::bSRGBEnabled;
-		SColorThemesViewer::bSRGBEnabled = InArgs._sRGBOverride.GetValue();
-	}
-
 	RegisterActiveTimer( 0.f, FWidgetActiveTimerDelegate::CreateSP( this, &SColorPicker::AnimatePostConstruct ) );
 
 	// We need a parent window to set the close callback
@@ -96,6 +82,7 @@ void SColorPicker::Construct( const FArguments& InArgs )
 	bColorPickerIsInlineVersion = InArgs._DisplayInlineVersion;
 	bIsInteractive = false;
 	bPerfIsTooSlowToUpdate = false;
+	bIsThemePanelVisible = true;
 
 	BeginAnimation(FLinearColor(ForceInit), CurrentColorHSV);
 
@@ -104,11 +91,16 @@ void SColorPicker::Construct( const FArguments& InArgs )
 		bool WheelMode = true;
 
 		GConfig->GetBool(TEXT("ColorPickerUI"), TEXT("bWheelMode"), WheelMode, GEditorPerProjectIni);
-		GConfig->GetBool(TEXT("ColorPickerUI"), TEXT("bSRGBEnabled"), SColorThemesViewer::bSRGBEnabled, GEditorPerProjectIni);
+		GConfig->GetBool(TEXT("ColorPickerUI"), TEXT("bSRGBEnabled"), bUseSRGB, GEditorPerProjectIni);
 
 		CurrentMode = WheelMode ? EColorPickerModes::Wheel : EColorPickerModes::Spectrum;
 	}
 
+	if (InArgs._sRGBOverride.IsSet())
+	{
+		bUseSRGB = InArgs._sRGBOverride.GetValue();
+	}
+
 	if (bColorPickerIsInlineVersion)
 	{
 		GenerateInlineColorPickerContent();
@@ -130,38 +122,6 @@ void SColorPicker::GenerateDefaultColorPickerContent( bool bAdvancedSectionExpan
 	// The height of the gradient bars beneath the sliders
 	const FSlateFontInfo SmallFont = FAppStyle::Get().GetFontStyle("ColorPicker.SmallFont");
 
-	TSharedPtr<SColorThemesViewer> ThemesViewer = ColorThemesViewer.Pin();
-
-	if (!ThemesViewer.IsValid())
-	{
-		ThemesViewer = SNew(SColorThemesViewer);
-		ColorThemesViewer = ThemesViewer;
-	}
-
-	ThemesViewer->OnCurrentThemeChanged().AddSP(this, &SColorPicker::HandleThemesViewerThemeChanged);
-	ThemesViewer->SetUseAlpha(bUseAlpha);
-	ThemesViewer->MenuToStandardNoReturn();
-
-	// The standard button to open the color themes can temporarily become a trash for colors
-	ColorThemeComboButton = SNew(SComboButton)
-		.ComboButtonStyle(FAppStyle::Get(), "SimpleComboButton")
-		.VAlign(VAlign_Center)
-		.ButtonContent()
-		[
-			SNew(SImage)
-			.Image(FAppStyle::Get().GetBrush("ColorPicker.ColorThemes"))
-		];
-
-	ColorThemeComboButton->SetMenuContent(ThemesViewer.ToSharedRef());
-
-	SmallTrash = SNew(SHorizontalBox)
-	+ SHorizontalBox::Slot()
-		.AutoWidth()
-		[
-			SNew(SColorTrash)
-				.UsesSmallIcon(true)
-		];
-	
 	this->ChildSlot
 	[
 		SNew(SVerticalBox)
@@ -313,6 +273,30 @@ void SColorPicker::GenerateDefaultColorPickerContent( bool bAdvancedSectionExpan
 											.Visibility(bValidCreationOverrideExists ? EVisibility::Collapsed : EVisibility::Visible)
 									]
 							]
+
+						+ SVerticalBox::Slot()
+							.MinHeight(28.0f)
+							.MaxHeight(28.0f)
+							.Padding(0.0f, 8.0f, 0.0f, 0.0f)
+							[
+								SNew(SHorizontalBox)
+
+								+ SHorizontalBox::Slot()
+									.MinWidth(48.0f)
+									.MaxWidth(48.0f)
+									[
+										// Show/Hide Themes Panel
+										SNew(SButton)
+											.OnClicked(this, &SColorPicker::ToggleThemePanelVisibility)
+											.ContentPadding(FMargin(2.0f, 2.5f))
+											.ToolTipText(LOCTEXT("ShowHideThemesButtonTooltip", "Toggle visibility of color themes"))
+											.Content()
+											[
+												SNew(SImage)
+													.Image(this, &SColorPicker::HandleThemePanelButtonImageBrush)
+											]
+									]
+							]
 					]
 			]
 
@@ -449,35 +433,18 @@ void SColorPicker::GenerateDefaultColorPickerContent( bool bAdvancedSectionExpan
 			.AutoHeight()
 			.Padding(0.0f, 16.0f, 0.0f, 0.0f)
 			[
-				SNew(SHorizontalBox)
-
-				+ SHorizontalBox::Slot()
-					.FillWidth(0.91f)
-					[
-						SNew(SOverlay)
-
-						+ SOverlay::Slot()
-							[
-								// color theme bar
-								SAssignNew(CurrentThemeBar, SThemeColorBlocksBar)
-									.ColorTheme(this, &SColorPicker::HandleThemeBarColorTheme)
-									.EmptyText(LOCTEXT("EmptyBarHint", "Drag & drop colors here to save"))
-									.HideTrashCallback(this, &SColorPicker::HideSmallTrash)
-									.ShowTrashCallback(this, &SColorPicker::ShowSmallTrash)
-									.ToolTipText(LOCTEXT("CurrentThemeBarToolTip", "Current Color Theme"))
-									.UseAlpha(SharedThis(this), &SColorPicker::HandleThemeBarUseAlpha)
-									.UseSRGB(SharedThis(this), &SColorPicker::HandleColorPickerUseSRGB)
-									.OnSelectColor(this, &SColorPicker::HandleThemeBarColorSelected)
-							]
-					]
-
-				+ SHorizontalBox::Slot()
-					.FillWidth(0.09f)
-					[
-						// color theme selector
-						SAssignNew(ColorThemeButtonOrSmallTrash, SBorder)
-							.BorderImage(FStyleDefaults::GetNoBrush())
-							.Padding(0.0f)
+ 				SNew(SBorder)
+ 					.BorderImage(FAppStyle::Get().GetBrush("ColorPicker.RecessedBackground"))
+ 					.Padding(FMargin(8.0f, 8.0f))
+					.Visibility(this, &SColorPicker::HandleThemesPanelVisibility)
+ 					[
+						// color theme bar
+						SAssignNew(CurrentThemeBar, SThemeColorBlocksBar)
+							.ToolTipText(this, &SColorPicker::GetColorThemePanelToolTipText)
+							.UseAlpha(bUseAlpha.Get())
+							.UseSRGB(SharedThis(this), &SColorPicker::HandleColorPickerUseSRGB)
+							.OnSelectColor(this, &SColorPicker::HandleThemeBarColorSelected)
+							.OnGetActiveColor(this, &SColorPicker::GetCurrentColor)
 					]
 			]
 			
@@ -514,8 +481,6 @@ void SColorPicker::GenerateDefaultColorPickerContent( bool bAdvancedSectionExpan
 					]
 			]
 	];
-	
-	HideSmallTrash();
 }
 END_SLATE_FUNCTION_BUILD_OPTIMIZATION
 
@@ -698,27 +663,13 @@ void SColorPicker::BeginAnimation( FLinearColor Start, FLinearColor End )
 
 void SColorPicker::HideSmallTrash()
 {
-	if (ColorThemeComboButton.IsValid())
-	{
-		ColorThemeButtonOrSmallTrash->SetContent(ColorThemeComboButton.ToSharedRef());
-	}
-	else
-	{
-		ColorThemeButtonOrSmallTrash->ClearContent();
-	}
+	// Deprecated function
 }
 
 
 void SColorPicker::ShowSmallTrash()
 {
-	if ( SmallTrash.IsValid() )
-	{
-		ColorThemeButtonOrSmallTrash->SetContent(SmallTrash.ToSharedRef());
-	}
-	else
-	{
-		ColorThemeButtonOrSmallTrash->ClearContent();
-	}
+	// Deprecated function
 }
 
 
@@ -852,12 +803,8 @@ TSharedRef<SWidget> SColorPicker::MakeColorSpinBox( EColorPickerChannels Channel
 }
 
 
-TSharedRef<SWidget> SColorPicker::MakeColorPreviewBox() const
+TSharedRef<SWidget> SColorPicker::MakeColorPreviewBox()
 {
-	const float OpaqueWidth = bUseAlpha.Get() ? 53.0f : 106.0f;
-	const FVector4 OldOpaqueBlockCornerRadius = bUseAlpha.Get() ? FVector4(4.0, 0.0, 0.0, 0.0) : FVector4(4.0, 4.0, 0.0, 0.0);
-	const FVector4 NewOpaqueBlockCornerRadius = bUseAlpha.Get() ? FVector4(0.0, 0.0, 0.0, 4.0) : FVector4(0.0, 0.0, 4.0, 4.0);
-
 	return SNew(SVerticalBox)
 
 	+ SVerticalBox::Slot()
@@ -867,81 +814,70 @@ TSharedRef<SWidget> SColorPicker::MakeColorPreviewBox() const
 
 			+ SOverlay::Slot()
 				[
-					// old color (alpha)
+					// new color (alpha)
 					SNew(SColorBlock)
 						.ColorIsHSV(true)
 						.ShowBackgroundForAlpha(true)
-						.ToolTipText(LOCTEXT("OldColorAlphaToolTip", "Old color with alpha (drag to theme bar to save)"))
-						.Color(OldColor)
-						.Visibility(SharedThis(this), &SColorPicker::HandleAlphaColorBlockVisibility)
-						.OnMouseButtonDown(const_cast<SColorPicker*>(this), &SColorPicker::HandleOldColorBlockMouseButtonDown, true)
+						.AlphaDisplayMode(SharedThis(this), &SColorPicker::HandleColorPreviewAlphaMode)
+						.Color(this, &SColorPicker::GetCurrentColor)
+						.ToolTipText(LOCTEXT("NewColorButtonToolTip", "Preview of the currently selected color"))
 						.UseSRGB(SharedThis(this), &SColorPicker::HandleColorPickerUseSRGB)
-						.Cursor(EMouseCursor::GrabHand)
 						.Size(FVector2D(106.0, 32.0))
-						.CornerRadius(FVector4(4.0, 4.0, 0.0, 0.0))
+						.CornerRadius(FVector4(4.0, 4.0, 4.0, 4.0))
 				]
 
 			+ SOverlay::Slot()
 				[
-					SNew(SHorizontalBox)
-
-					+SHorizontalBox::Slot()
-						.AutoWidth()
+					SAssignNew(NewColorButton, SButton)
+						.ButtonStyle(FAppStyle::Get(), "ColorPicker.ColorPreviewButton")
+						.OnClicked(const_cast<SColorPicker*>(this), &SColorPicker::HandleNewColorPreviewClicked)
+						.HAlign(HAlign_Center)
+						.VAlign(VAlign_Center)
+						.ToolTipText(LOCTEXT("NewColorButtonToolTip", "Add the currently selected color to the current color theme"))
+						.Visibility(this, &SColorPicker::HandleColorPreviewButtonVisibility)
+						.Content()
 						[
-							// old color
-							SNew(SColorBlock)
-								.ColorIsHSV(true)
-								.AlphaDisplayMode(EColorBlockAlphaDisplayMode::Ignore)
-								.ToolTipText(LOCTEXT("OldColorToolTip", "Old color without alpha (drag to theme bar to save)"))
-								.Color(OldColor)
-								.OnMouseButtonDown(const_cast<SColorPicker*>(this), &SColorPicker::HandleOldColorBlockMouseButtonDown, false)
-								.UseSRGB(SharedThis(this), &SColorPicker::HandleColorPickerUseSRGB)
-								.Cursor(EMouseCursor::GrabHand)
-								.Size(FVector2D(OpaqueWidth, 32.0))
-								.CornerRadius(OldOpaqueBlockCornerRadius)
+							SNew(SImage)
+								.Image(FAppStyle::Get().GetBrush("Icons.Plus"))
+								.Visibility(const_cast<SColorPicker*>(this), &SColorPicker::HandleNewColorButtonVisibility)
 						]
 				]
 		]
 
 	+ SVerticalBox::Slot()
 		.AutoHeight()
+		.Padding(0.0f, 2.0f, 0.0f, 0.0f)
 		[
 			SNew(SOverlay)
 
 			+ SOverlay::Slot()
 				[
-					// new color (alpha)
+					// Old color
 					SNew(SColorBlock)
 						.ColorIsHSV(true)
 						.ShowBackgroundForAlpha(true)
-						.ToolTipText(LOCTEXT("NewColorAlphaToolTip", "New color with alpha (drag to theme bar to save)"))
-						.Color(this, &SColorPicker::GetCurrentColor)
-						.Visibility(SharedThis(this), &SColorPicker::HandleAlphaColorBlockVisibility)
-						.OnMouseButtonDown(const_cast<SColorPicker*>(this), &SColorPicker::HandleNewColorBlockMouseButtonDown, true)
+						.AlphaDisplayMode(SharedThis(this), &SColorPicker::HandleColorPreviewAlphaMode)
+						.Color(OldColor)
+						.ToolTipText(LOCTEXT("NewColorButtonToolTip", "Preview of the previously selected color"))
 						.UseSRGB(SharedThis(this), &SColorPicker::HandleColorPickerUseSRGB)
-						.Cursor(EMouseCursor::GrabHand)
 						.Size(FVector2D(106.0, 32.0))
-						.CornerRadius(FVector4(0.0, 0.0, 4.0, 4.0))
+						.CornerRadius(FVector4(4.0, 4.0, 4.0, 4.0))
 				]
 
 			+ SOverlay::Slot()
 				[
-					SNew(SHorizontalBox)
-
-					+ SHorizontalBox::Slot()
-						.AutoWidth()
+					SAssignNew(OldColorButton, SButton)
+						.ButtonStyle(FAppStyle::Get(), "ColorPicker.ColorPreviewButton")
+						.OnClicked(const_cast<SColorPicker*>(this), &SColorPicker::HandleOldColorPreviewClicked)
+						.HAlign(HAlign_Center)
+						.VAlign(VAlign_Center)
+						.ToolTipText(LOCTEXT("OldColorButtonToolTip", "Add the previously selected color to the current color theme"))
+						.Visibility(this, &SColorPicker::HandleColorPreviewButtonVisibility)
+						.Content()
 						[
-							// new color
-							SNew(SColorBlock)
-								.ColorIsHSV(true)
-								.AlphaDisplayMode(EColorBlockAlphaDisplayMode::Ignore)
-								.ToolTipText(LOCTEXT("NewColorToolTip", "New color without alpha (drag to theme bar to save)"))
-								.Color(this, &SColorPicker::GetCurrentColor)
-								.OnMouseButtonDown(const_cast<SColorPicker*>(this), &SColorPicker::HandleNewColorBlockMouseButtonDown, false)
-								.UseSRGB(SharedThis(this), &SColorPicker::HandleColorPickerUseSRGB)
-								.Cursor(EMouseCursor::GrabHand)
-								.Size(FVector2D(OpaqueWidth, 32.0))
-								.CornerRadius(NewOpaqueBlockCornerRadius)
+							SNew(SImage)
+								.Image(FAppStyle::Get().GetBrush("Icons.Plus"))
+								.Visibility(const_cast<SColorPicker*>(this), &SColorPicker::HandleOldColorButtonVisibility)
 						]
 				]
 		];
@@ -1000,9 +936,9 @@ TArray<FLinearColor> SColorPicker::GetGradientColors( EColorPickerChannels Chann
 }
 
 
-EVisibility SColorPicker::HandleAlphaColorBlockVisibility() const
+EColorBlockAlphaDisplayMode SColorPicker::HandleColorPreviewAlphaMode() const
 {
-	return bUseAlpha.Get() ? EVisibility::Visible : EVisibility::Collapsed;
+	return bUseAlpha.Get() ? EColorBlockAlphaDisplayMode::SeparateReverse : EColorBlockAlphaDisplayMode::Ignore;
 }
 
 
@@ -1030,6 +966,34 @@ EVisibility SColorPicker::HandleColorPickerModeVisibility( EColorPickerModes Mod
 }
 
 
+EVisibility SColorPicker::HandleThemesPanelVisibility() const
+{
+	return bIsThemePanelVisible ? EVisibility::Visible : EVisibility::Collapsed;
+}
+
+FText SColorPicker::GetColorThemePanelToolTipText() const
+{
+	return (CurrentThemeBar->IsRecentsThemeActive()) ? LOCTEXT("RecentsThemeToolTipText", "Recently used colors") : LOCTEXT("ColorThemeToolTipText", "Current Color Theme");
+}
+
+const FSlateBrush* SColorPicker::HandleThemePanelButtonImageBrush() const
+{
+	return bIsThemePanelVisible ? FAppStyle::Get().GetBrush("ColorPicker.ColorThemes") : FAppStyle::Get().GetBrush("ColorPicker.ColorThemesOff");
+}
+
+void SColorPicker::HandleThemeBarColorSelected(FLinearColor NewValue)
+{
+	// Force the alpha component to 1 when we don't care about the alpha
+	if (!bUseAlpha.Get())
+	{
+		NewValue.A = 1.0f;
+	}
+
+	BeginAnimation(CurrentColorHSV, NewValue);
+	SetNewTargetColorHSV(NewValue, true);
+}
+
+
 FLinearColor SColorPicker::HandleColorSliderEndColor( EColorPickerChannels Channel ) const
 {
 	switch (Channel)
@@ -1246,22 +1210,42 @@ FReply SColorPicker::HandleColorPickerModeButtonClicked()
 }
 
 
-FReply SColorPicker::HandleNewColorBlockMouseButtonDown( const FGeometry& MyGeometry, const FPointerEvent& MouseEvent, bool bCheckAlpha )
+EVisibility SColorPicker::HandleColorPreviewButtonVisibility() const
 {
-	if (MouseEvent.GetEffectingButton() == EKeys::LeftMouseButton)
-	{
-		TSharedRef<FColorDragDrop> Operation = FColorDragDrop::New(
-			CurrentColorHSV,
-			SColorThemesViewer::bSRGBEnabled,
-			bCheckAlpha ? bUseAlpha.Get() : false,
-			FSimpleDelegate::CreateSP(this, &SColorPicker::ShowSmallTrash),
-			FSimpleDelegate::CreateSP(this, &SColorPicker::HideSmallTrash)
-		);
-
-		return FReply::Handled().BeginDragDrop(Operation);
-	}
+	return CurrentThemeBar->IsRecentsThemeActive() ? EVisibility::Hidden : EVisibility::Visible;
+}
 
-	return FReply::Unhandled();
+
+EVisibility SColorPicker::HandleNewColorButtonVisibility() const
+{
+	return NewColorButton->IsHovered() ? EVisibility::Visible : EVisibility::Hidden;
+}
+
+
+EVisibility SColorPicker::HandleOldColorButtonVisibility() const
+{
+	return OldColorButton->IsHovered() ? EVisibility::Visible : EVisibility::Hidden;
+}
+
+
+FReply SColorPicker::HandleNewColorPreviewClicked()
+{
+	CurrentThemeBar->AddNewColorBlock(CurrentColorHSV, 0);
+	return FReply::Handled();
+}
+
+
+FReply SColorPicker::HandleOldColorPreviewClicked()
+{
+	CurrentThemeBar->AddNewColorBlock(OldColor, 0);
+	return FReply::Handled();
+}
+
+
+FReply SColorPicker::ToggleThemePanelVisibility()
+{
+	bIsThemePanelVisible = !bIsThemePanelVisible;
+	return FReply::Handled();
 }
 
 
@@ -1271,6 +1255,11 @@ FReply SColorPicker::HandleOkButtonClicked()
 
 	UpdateColorPick();
 
+	if (OldColor != CurrentColorHSV)
+	{
+		CurrentThemeBar->AddToRecents(CurrentColorHSV);
+	}
+
 	if (SColorPicker::OnColorPickerDestroyOverride.IsBound())
 	{
 		SColorPicker::OnColorPickerDestroyOverride.Execute();
@@ -1283,28 +1272,9 @@ FReply SColorPicker::HandleOkButtonClicked()
 }
 
 
-FReply SColorPicker::HandleOldColorBlockMouseButtonDown( const FGeometry& MyGeometry, const FPointerEvent& MouseEvent, bool bCheckAlpha )
-{
-	if (MouseEvent.GetEffectingButton() == EKeys::LeftMouseButton)
-	{
-		TSharedRef<FColorDragDrop> Operation = FColorDragDrop::New(
-			OldColor,
-			SColorThemesViewer::bSRGBEnabled,
-			bCheckAlpha ? bUseAlpha.Get() : false,
-			FSimpleDelegate::CreateSP(this, &SColorPicker::ShowSmallTrash),
-			FSimpleDelegate::CreateSP(this, &SColorPicker::HideSmallTrash)
-		);
-
-		return FReply::Handled().BeginDragDrop(Operation);
-	}
-
-	return FReply::Unhandled();
-}
-
-
 bool SColorPicker::HandleColorPickerUseSRGB() const
 {
-	return SColorThemesViewer::bSRGBEnabled;
+	return bUseSRGB;
 }
 
 
@@ -1330,11 +1300,6 @@ void SColorPicker::HandleParentWindowClosed( const TSharedRef<SWindow>& Window )
 	{
 		DiscardColor();
 	}
-
-	if ( OriginalSRGBOption.IsSet() )
-	{
-		SColorThemesViewer::bSRGBEnabled = OriginalSRGBOption.GetValue();
-	}
 }
 
 
@@ -1346,81 +1311,21 @@ void SColorPicker::HandleRGBColorChanged( FLinearColor NewValue )
 
 void SColorPicker::HandleSRGBCheckBoxCheckStateChanged( ECheckBoxState InIsChecked )
 {
-	SColorThemesViewer::bSRGBEnabled = (InIsChecked == ECheckBoxState::Checked);
+	bUseSRGB = (InIsChecked == ECheckBoxState::Checked);
 
 	if (FPaths::FileExists(GEditorPerProjectIni))
 	{
-		GConfig->SetBool(TEXT("ColorPickerUI"), TEXT("bSRGBEnabled"), SColorThemesViewer::bSRGBEnabled, GEditorPerProjectIni);
+		GConfig->SetBool(TEXT("ColorPickerUI"), TEXT("bSRGBEnabled"), bUseSRGB, GEditorPerProjectIni);
 	}
 }
 
 
 ECheckBoxState SColorPicker::HandleSRGBCheckBoxIsChecked() const
 {
-	return SColorThemesViewer::bSRGBEnabled ? ECheckBoxState::Checked : ECheckBoxState::Unchecked;
-}
-
-
-void SColorPicker::HandleThemeBarColorSelected( FLinearColor NewValue )
-{
-	// Force the alpha component to 1 when we don't care about the alpha
-	if (!bUseAlpha.Get())
-	{
-		NewValue.A = 1.0f;
-	}
-
-	BeginAnimation(CurrentColorHSV, NewValue);
-	SetNewTargetColorHSV(NewValue, true);
+	return bUseSRGB ? ECheckBoxState::Checked : ECheckBoxState::Unchecked;
 }
 
 
-TSharedPtr<FColorTheme> SColorPicker::HandleThemeBarColorTheme() const
-{
-	TSharedPtr<SColorThemesViewer> ThemesViewer = ColorThemesViewer.Pin();
-
-	if (ThemesViewer.IsValid())
-	{
-		return ThemesViewer->GetCurrentColorTheme();
-	}
-
-	return nullptr;
-}
-
-
-EVisibility SColorPicker::HandleThemeBarHintVisibility() const
-{
-	TSharedPtr<SColorThemesViewer> ThemesViewer = ColorThemesViewer.Pin();
-
-	if (ThemesViewer.IsValid())
-	{
-		TSharedPtr<FColorTheme> SelectedTheme = ThemesViewer->GetCurrentColorTheme();
-
-		if (SelectedTheme.IsValid() && SelectedTheme->GetColors().Num() == 0)
-		{
-			return EVisibility::HitTestInvisible;
-		}
-	}
-
-	return EVisibility::Hidden;
-}
-
-
-bool SColorPicker::HandleThemeBarUseAlpha() const
-{
-	return bUseAlpha.Get();
-}
-
-
-void SColorPicker::HandleThemesViewerThemeChanged()
-{
-	if (CurrentThemeBar.IsValid())
-	{
-		CurrentThemeBar->RemoveRefreshCallback();
-		CurrentThemeBar->AddRefreshCallback();
-		CurrentThemeBar->Refresh();
-	}
-}
-
 // Static delegates to access whether or not the override is bound in the global Open/Destroy functions
 SColorPicker::FOnColorPickerCreationOverride SColorPicker::OnColorPickerNonModalCreateOverride;
 SColorPicker::FOnColorPickerDestructionOverride SColorPicker::OnColorPickerDestroyOverride;
diff --git a/Engine/Source/Runtime/AppFramework/Private/Widgets/Colors/SColorSlider.cpp b/Engine/Source/Runtime/AppFramework/Private/Widgets/Colors/SColorSlider.cpp
--- a/Engine/Source/Runtime/AppFramework/Private/Widgets/Colors/SColorSlider.cpp
+++ b/Engine/Source/Runtime/AppFramework/Private/Widgets/Colors/SColorSlider.cpp
@@ -156,16 +156,17 @@ int32 SColorSlider::OnPaint(const FPaintArgs& Args, const FGeometry& AllottedGeo
 		}
 
 		// Vertical sliders need the start color at the bottom and the end color at the bottom
+		TArray<FLinearColor> Colors = GradientColors.Get();
 		if (Orientation.Get() == Orient_Vertical)
 		{
-			Algo::Reverse(GradientColors.Get());
+			Algo::Reverse(Colors);
 		}
 
 		// Generate the gradient stops for the colors
 		TArray<FSlateGradientStop> GradientStops;
 		for (int32 ColorIndex = 0; ColorIndex < NumColors; ++ColorIndex)
 		{
-			GradientStops.Add(FSlateGradientStop(ColorSliderSize * (float(ColorIndex) / (NumColors - 1)), GradientColors.Get()[ColorIndex]));
+			GradientStops.Add(FSlateGradientStop(ColorSliderSize * (float(ColorIndex) / (NumColors - 1)), Colors[ColorIndex]));
 		}
 
 		const FVector4f CornerRadius = FVector4f(4.0f, 4.0f, 4.0f, 4.0f);
diff --git a/Engine/Source/Runtime/AppFramework/Private/Widgets/Colors/SColorThemes.cpp b/Engine/Source/Runtime/AppFramework/Private/Widgets/Colors/SColorThemes.cpp
--- a/Engine/Source/Runtime/AppFramework/Private/Widgets/Colors/SColorThemes.cpp
+++ b/Engine/Source/Runtime/AppFramework/Private/Widgets/Colors/SColorThemes.cpp
@@ -1,32 +1,30 @@
 // Copyright Epic Games, Inc. All Rights Reserved.
 
 #include "Widgets/Colors/SColorThemes.h"
-#include "Misc/ConfigCacheIni.h"
-#include "Layout/ArrangedChildren.h"
+
 #include "Application/SlateWindowHelper.h"
+#include "Framework/Application/SlateApplication.h"
+#include "Layout/ArrangedChildren.h"
+#include "Misc/ConfigCacheIni.h"
 #include "SlateOptMacros.h"
 #include "Styling/CoreStyle.h"
+#include "ToolMenus.h"
+#include "Widgets/Colors/SColorBlock.h"
 #include "Widgets/Images/SImage.h"
-#include "Widgets/Text/STextBlock.h"
-#include "Widgets/Input/SEditableTextBox.h"
 #include "Widgets/Input/SButton.h"
-#include "Widgets/SToolTip.h"
-#include "Widgets/Notifications/SErrorText.h"
-#include "Widgets/Views/SListView.h"
-#include "Widgets/Colors/SColorBlock.h"
+#include "Widgets/Input/SComboButton.h"
+#include "Widgets/Input/SEditableTextBox.h"
 #include "Widgets/Input/STextEntryPopup.h"
-#include "Framework/Application/SlateApplication.h"
+#include "Widgets/Text/STextBlock.h"
+#include "Widgets/SToolTip.h"
+
+
+#define LOCTEXT_NAMESPACE "SColorThemes"
 
 
 void FColorDragDrop::OnDrop( bool bDropWasHandled, const FPointerEvent& MouseEvent )
 {
 	HideTrash.ExecuteIfBound();
-
-	if (OriginBar.Pin().IsValid() && !bSetForDeletion)
-	{
-		OriginBar.Pin()->AddNewColorBlock(Color, OriginBarPosition);
-	}
-	
 	FDragDropOperation::OnDrop( bDropWasHandled, MouseEvent );
 }
 
@@ -34,7 +32,7 @@ void FColorDragDrop::OnDragged( const class FDragDropEvent& DragDropEvent )
 {
 	if (CursorDecoratorWindow.IsValid())
 	{
-		CursorDecoratorWindow->MoveWindowTo( DragDropEvent.GetScreenSpacePosition() - BlockSize * 0.5f );
+		CursorDecoratorWindow->MoveWindowTo(DragDropEvent.GetScreenSpacePosition() - (CursorDecoratorWindow->GetSizeInScreen() * 0.5f));
 	}
 }
 
@@ -42,41 +40,58 @@ TSharedRef<FColorDragDrop> FColorDragDrop::New(FLinearColor InColor, bool bSRGB,
 	FSimpleDelegate TrashShowCallback, FSimpleDelegate TrashHideCallback,
 	TSharedPtr<SThemeColorBlocksBar> Origin, int32 OriginPosition)
 {
-	TSharedRef<FColorDragDrop> Operation = MakeShareable(new FColorDragDrop);
-
-	Operation->Color = InColor;
-	Operation->bUseSRGB = bSRGB;
-	Operation->bUseAlpha = bUseAlpha;
-	Operation->OriginBar = Origin;
-	Operation->OriginBarPosition = OriginPosition;
-	Operation->ShowTrash = TrashShowCallback;
-	Operation->HideTrash = TrashHideCallback;
-	Operation->bSetForDeletion = false;
-	Operation->BlockSize = FVector2D(32, 32);
-
-	Operation->ShowTrash.ExecuteIfBound();
-
-	Operation->Construct();
-
+	TSharedRef<FColorDragDrop> Operation = MakeShareable(new FColorDragDrop(InColor, bSRGB, bUseAlpha, TrashShowCallback, TrashHideCallback));
 	return Operation;
 }
 
 TSharedPtr<SWidget> FColorDragDrop::GetDefaultDecorator() const
 {
-	const bool bIgnoreAlpha = !bUseAlpha;
-	const bool bShowBackgroundForAlpha = bUseAlpha;
-
-	return SNew(SBorder) .Cursor(EMouseCursor::GrabHandClosed)
+	return SNew(SBorder)
+		.BorderImage(FAppStyle::Get().GetBrush("ColorPicker.MultipleValuesBackground"))
+		.Padding(FMargin(1.0f, 1.0f))
 		[
 			SNew(SColorBlock) 
 				.Color(Color)
 				.ColorIsHSV(true) 
-				.AlphaDisplayMode(bIgnoreAlpha ? EColorBlockAlphaDisplayMode::Ignore : EColorBlockAlphaDisplayMode::Combined)
-				.ShowBackgroundForAlpha(bShowBackgroundForAlpha) 
+				.AlphaDisplayMode(bUseAlpha ? EColorBlockAlphaDisplayMode::SeparateReverse : EColorBlockAlphaDisplayMode::Ignore)
+				.ShowBackgroundForAlpha(bUseAlpha)
 				.UseSRGB(bUseSRGB)
+				.Size(FVector2D(22.0, 22.0))
+				.CornerRadius(FVector4(4.0, 4.0, 4.0, 4.0))
 		];
 }
 
+FColorDragDrop::FColorDragDrop(FLinearColor InColor, bool bInUseSRGB, bool bInUseAlpha, FSimpleDelegate InTrashShowCallback, FSimpleDelegate InTrashHideCallback)
+{
+	Color = InColor;
+	bUseSRGB = bInUseSRGB;
+	bUseAlpha = bInUseAlpha;
+	ShowTrash = InTrashShowCallback;
+	HideTrash = InTrashHideCallback;
+
+	TSharedPtr<SWidget> DecoratorToUse = GetDefaultDecorator();
+
+	if (DecoratorToUse.IsValid())
+	{
+		CursorDecoratorWindow = SWindow::MakeStyledCursorDecorator(FAppStyle::Get().GetWidgetStyle<FWindowStyle>("ColorPicker.CursorDecorator"));
+		CursorDecoratorWindow->SetContent(DecoratorToUse.ToSharedRef());
+
+		FSlateApplicationBase::Get().AddWindow(CursorDecoratorWindow.ToSharedRef(), true);
+	}
+
+	ShowTrash.ExecuteIfBound();
+}
+
+void FColorDragDrop::MarkForDelete()
+{
+	CursorDecoratorWindow->SetOpacity(0.4f);
+}
+
+void FColorDragDrop::MarkForAdd()
+{
+	CursorDecoratorWindow->SetOpacity(1.0f);
+}
+
 
 FColorTheme::FColorTheme( const FString& InName, const TArray< TSharedPtr<FColorInfo> >& InColors )
 	: Name(InName)
@@ -97,7 +112,6 @@ void FColorTheme::InsertNewColor( TSharedPtr<FLinearColor> InColor, int32 Insert
 	RefreshEvent.Broadcast();
 }
 
-
 int32 FColorTheme::FindApproxColor( const FLinearColor& InColor, float Tolerance ) const
 {
 	for (int32 ColorIndex = 0; ColorIndex < Colors.Num(); ++ColorIndex)
@@ -112,15 +126,12 @@ int32 FColorTheme::FindApproxColor( const FLinearColor& InColor, float Tolerance
 	return INDEX_NONE;
 }
 
-
 void FColorTheme::RemoveAll()
 {
 	Colors.Empty();
 	RefreshEvent.Broadcast();
 }
 
-
-
 int32 FColorTheme::RemoveColor( const TSharedPtr<FLinearColor> InColor )
 {
 	int32 Index = INDEX_NONE;
@@ -134,7 +145,7 @@ int32 FColorTheme::RemoveColor( const TSharedPtr<FLinearColor> InColor )
 	return Index;
 }
 
-
+PRAGMA_DISABLE_DEPRECATION_WARNINGS
 void SColorTrash::Construct( const FArguments& InArgs )
 {
 	bBorderActivated = false;
@@ -217,8 +228,7 @@ const FSlateBrush* SColorTrash::GetBorderStyle() const
 		FAppStyle::Get().GetBrush("FocusRectangle") :
 		FStyleDefaults::GetNoBrush();
 }
-
-
+PRAGMA_ENABLE_DEPRECATION_WARNINGS
 
 /**
 * Construct the widget
@@ -235,8 +245,7 @@ void SThemeColorBlock::Construct(const FArguments& InArgs )
 	HideTrashCallback = InArgs._HideTrashCallback;
 	bUseSRGB = InArgs._UseSRGB;
 	bUseAlpha = InArgs._UseAlpha;
-
-	DistanceDragged = 0;
+	bSupportsDrag = InArgs._SupportsDrag;
 
 	const FSlateFontInfo SmallLayoutFont = FAppStyle::Get().GetFontStyle("Regular");
 	const FSlateFontInfo SmallLabelFont = FAppStyle::Get().GetFontStyle("Bold");
@@ -314,19 +323,18 @@ void SThemeColorBlock::Construct(const FArguments& InArgs )
 	this->ChildSlot
 	[
 		SNew(SBorder)
-			.BorderImage(this, &SThemeColorBlock::HandleBorderImage)
-			.BorderBackgroundColor(this, &SThemeColorBlock::HandleBorderColor)
+			.BorderImage(FAppStyle::Get().GetBrush("ColorPicker.MultipleValuesBackground"))
 			.Padding(FMargin(1.0f))
 			.ToolTip(ColorTooltip)
-			.VAlign(VAlign_Center)
 			[
 				SNew(SColorBlock)
 				.Color(this, &SThemeColorBlock::GetColor)
-				.AlphaDisplayMode(EColorBlockAlphaDisplayMode::Separate)
+				.AlphaDisplayMode(bUseAlpha.Get() ? EColorBlockAlphaDisplayMode::SeparateReverse : EColorBlockAlphaDisplayMode::Ignore)
 				.ColorIsHSV(true)
-				.ShowBackgroundForAlpha(true)
+				.ShowBackgroundForAlpha(bUseAlpha.Get())
 				.UseSRGB(bUseSRGB)
-				.Size(FVector2D(32,16))
+				.Size(FVector2D(22.0, 22.0))
+				.CornerRadius(FVector4(4.0, 4.0, 4.0, 4.0))
 			]
 	];
 }
@@ -368,9 +376,8 @@ void SThemeColorBlock::SetLabel(const FText& NewColorLabel, ETextCommit::Type Co
 
 FReply SThemeColorBlock::OnMouseButtonDown( const FGeometry& MyGeometry, const FPointerEvent& MouseEvent )
 {
-	if ( MouseEvent.GetEffectingButton() == EKeys::LeftMouseButton )
+	if ( MouseEvent.GetEffectingButton() == EKeys::LeftMouseButton && bSupportsDrag.Get())
 	{
-		DistanceDragged = 0;
 		return FReply::Handled().DetectDrag( SharedThis(this), EKeys::LeftMouseButton ).CaptureMouse(SharedThis(this));
 	}
 	if (MouseEvent.GetEffectingButton() == EKeys::RightMouseButton)
@@ -389,7 +396,7 @@ FReply SThemeColorBlock::OnMouseButtonUp( const FGeometry& MyGeometry, const FPo
 	if ( MouseEvent.GetEffectingButton() == EKeys::LeftMouseButton && MyGeometry.IsUnderLocation(MouseEvent.GetScreenSpacePosition()) )
 	{
 		check(ColorPtr.IsValid());
-		OnSelectColor.ExecuteIfBound(*ColorPtr.Pin());
+		OnSelectColor.ExecuteIfBound(GetColor());
 		
 		return FReply::Handled().ReleaseMouseCapture();
 			
@@ -402,16 +409,14 @@ FReply SThemeColorBlock::OnMouseButtonUp( const FGeometry& MyGeometry, const FPo
 
 FReply SThemeColorBlock::OnDragDetected( const FGeometry& MyGeometry, const FPointerEvent& MouseEvent )
 {
-	if ( MouseEvent.IsMouseButtonDown( EKeys::LeftMouseButton ) && ParentPtr.IsValid() )
+	if ( MouseEvent.IsMouseButtonDown( EKeys::LeftMouseButton ) && bSupportsDrag.Get())
 	{
-		FLinearColor ColorToGrab = *ColorPtr.Pin();
+		if (ParentPtr.IsValid())
+		{
+			ParentPtr.Pin()->RemoveColorBlock(ColorPtr.Pin());
+		}
 
-		int32 Position = ParentPtr.Pin()->RemoveColorBlock(ColorPtr.Pin());
-			
-		ParentPtr.Pin()->SetPlaceholderGrabOffset(MyGeometry.AbsoluteToLocal( MouseEvent.GetScreenSpacePosition() ));
-		
-		TSharedRef<FColorDragDrop> Operation = FColorDragDrop::New(ColorToGrab, bUseSRGB.Get(), bUseAlpha.Get(), 
-			ShowTrashCallback, HideTrashCallback, ParentPtr.Pin(), Position);
+		TSharedRef<FColorDragDrop> Operation = FColorDragDrop::New(GetColor(), bUseSRGB.Get(), bUseAlpha.Get(), ShowTrashCallback, HideTrashCallback);
 		return FReply::Handled().BeginDragDrop(Operation);
 	}
 	
@@ -423,24 +428,6 @@ FLinearColor SThemeColorBlock::GetColor() const
 	return ColorPtr.IsValid() ? *ColorPtr.Pin() : FLinearColor(ForceInit);
 }
 
-FSlateColor SThemeColorBlock::HandleBorderColor() const
-{
-	if (IsHovered())
-	{
-		return FSlateColor(FLinearColor::White);
-	}
-	
-	return FSlateColor(GetColor().HSVToLinearRGB().ToFColor(bUseSRGB.Get()));
-}
-
-const FSlateBrush* SThemeColorBlock::HandleBorderImage() const
-{
-	return IsHovered() ?
-		FAppStyle::Get().GetBrush("FocusRectangle") :
-		FAppStyle::Get().GetBrush("GenericWhiteBox");
-}
-
-#define LOCTEXT_NAMESPACE "ThemeColorBlock"
 FText SThemeColorBlock::GetRedText() const { return ColorPtr.IsValid() ? FormatToolTipText(LOCTEXT("Red", "R"), ColorPtr.Pin()->HSVToLinearRGB().R) : FText::GetEmpty(); }
 FText SThemeColorBlock::GetGreenText() const { return ColorPtr.IsValid() ? FormatToolTipText(LOCTEXT("Green", "G"), ColorPtr.Pin()->HSVToLinearRGB().G) : FText::GetEmpty(); }
 FText SThemeColorBlock::GetBlueText() const { return ColorPtr.IsValid() ? FormatToolTipText(LOCTEXT("Blue", "B"), ColorPtr.Pin()->HSVToLinearRGB().B) : FText::GetEmpty(); }
@@ -473,7 +460,6 @@ FText SThemeColorBlock::FormatToolTipText(const FText& ColorIdentifier, float Va
 
 	return FText::Format(LOCTEXT("ToolTipFormat", "{Identifier}: {Value}"), Args);
 }
-#undef LOCTEXT_NAMESPACE
 
 EColorBlockAlphaDisplayMode SThemeColorBlock::OnGetAlphaDisplayMode() const
 {
@@ -495,82 +481,75 @@ EVisibility SThemeColorBlock::OnGetLabelVisibility() const
 	return !ColorInfo->Label.IsEmpty() ? EVisibility::Visible : EVisibility::Collapsed;
 }
 
-
 SThemeColorBlocksBar::SThemeColorBlocksBar()
-: ColorBlocks(this)
+	: Children(this)
 {
 }
 
-
-/**
-* Panels arrange their children in a space described by the AllottedGeometry parameter. The results of the arrangement
-* should be returned by appending a FArrangedWidget pair for every child widget.
-*
-* @param AllottedGeometry    The geometry allotted for this widget by its parent.
-* @param ArrangedChildren    The array to which to add the WidgetGeometries that represent the arranged children.
-*/
 void SThemeColorBlocksBar::OnArrangeChildren( const FGeometry& AllottedGeometry, FArrangedChildren& ArrangedChildren ) const
 {
-	bool bPlaceHolderExists = NewColorPlaceholder.IsValid();
+	bool bPlaceholderExists = PlaceholderIndex.IsSet();
+
+	const FVector2D BlockSize = FVector2D(24.0f, 24.0f);
+	constexpr float Padding = 2.0f;
 
-	const int32 NumColorBlocks = bPlaceHolderExists ? ColorBlocks.Num() + 1 : ColorBlocks.Num();
+	// Add the combo button, which takes up two grid slots (with padding)
+	ArrangedChildren.AddWidget(AllottedGeometry.MakeChild(ThemesViewer.ToSharedRef(), FVector2D(0.0, 0.0), FVector2D(50.0, 24.0)));
+	int32 OccupiedGridSlots = 2;
 
-	if (NumColorBlocks > 0)
+	// Add the add/delete button, which takes up one grid slot
+	if (!ThemesViewer->IsRecentsThemeActive())
 	{
-		const float ChildPadding = 3.0f;
-		const FVector2D ChildSize = FVector2D(AllottedGeometry.GetLocalSize().X / NumColorBlocks - ChildPadding, AllottedGeometry.GetLocalSize().Y);
+		ArrangedChildren.AddWidget(AllottedGeometry.MakeChild(AddDeleteOverlay.ToSharedRef(), FVector2D(52.0, 0.0), BlockSize));
+		OccupiedGridSlots++;
+	}
 
-		const float CurrentDragCenter = PlaceholderBlockOffset + ChildSize.X * 0.5f;
+	const int32 NumColorBlocks = bPlaceholderExists ? ColorBlocks.Num() + 1 : ColorBlocks.Num();
+	const int32 NumGridBlocks = NumColorBlocks + OccupiedGridSlots;
 
-		float XOffset = 0;
-		for (int32 i = 0; i < NumColorBlocks; ++i)
-		{
-			if (bPlaceHolderExists && XOffset <= CurrentDragCenter && CurrentDragCenter < (XOffset + ChildSize.X))
-			{
-				XOffset += ChildSize.X;
-				XOffset += ChildPadding;
-			}
+	int32 ColorIndex = 0;
+	for (int32 GridIndex = OccupiedGridSlots; GridIndex < NumGridBlocks; ++GridIndex)
+	{
+		const float HPadding = (GridIndex % 16) ? Padding : 0.0f;
+		const float VPadding = (GridIndex / 16) ? Padding : 0.0f;
 
-			if (i < ColorBlocks.Num())
-			{
-				ArrangedChildren.AddWidget(AllottedGeometry.MakeChild(ColorBlocks[i], FVector2D(XOffset, 0), ChildSize));
-				XOffset += ChildSize.X;
-				XOffset += ChildPadding;
-			}
-		}
+		float YOffset = (GridIndex / 16) * (BlockSize.Y + Padding);
+		float XOffset = (GridIndex % 16) * (BlockSize.X + Padding);
+		FVector2D DragShadowLocation = FVector2D(XOffset, YOffset);
 
-		if (bPlaceHolderExists && NewColorBlockPlaceholder.IsValid())
+		if (PlaceholderIndex.IsSet() && (GridIndex == PlaceholderIndex))
 		{
-			ArrangedChildren.AddWidget(AllottedGeometry.MakeChild(NewColorBlockPlaceholder.ToSharedRef(), FVector2D(PlaceholderBlockOffset, 0), ChildSize));
+			ArrangedChildren.AddWidget(AllottedGeometry.MakeChild(DragShadow.ToSharedRef(), DragShadowLocation, BlockSize));
+		}
+		else if (ColorIndex < ColorBlocks.Num())
+		{
+			ArrangedChildren.AddWidget(AllottedGeometry.MakeChild(ColorBlocks[ColorIndex].ToSharedRef(), FVector2D(XOffset, YOffset), BlockSize));
+			ColorIndex++;
 		}
-	}
-	else if (EmptyHintTextBlock.IsValid())
-	{
-		ArrangedChildren.AddWidget(AllottedGeometry.MakeChild(EmptyHintTextBlock.ToSharedRef(), FVector2D::ZeroVector, AllottedGeometry.Size));
 	}
 }
 
-/**
-* A Panel's desired size in the space required to arrange of its children on the screen while respecting all of
-* the children's desired sizes and any layout-related options specified by the user. See StackPanel for an example.
-*
-* @return The desired size.
-*/
 FVector2D SThemeColorBlocksBar::ComputeDesiredSize( float ) const
 {
-	return FVector2D(64, 16);
+	const bool bPlaceholderExists = PlaceholderIndex.IsSet();
+	const bool bAddDeleteButtonVisible = !ThemesViewer->IsRecentsThemeActive();
+	const int32 NumColorBlocks = bPlaceholderExists ? ColorBlocks.Num() + 1 : ColorBlocks.Num();
+	const int32 NumGridBlocks = bAddDeleteButtonVisible ? NumColorBlocks + 3 : NumColorBlocks + 2;
+
+	const FVector2D BlockSize = FVector2D(24.0f, 24.0f);
+	constexpr float Padding = 2.0f;
+
+	int32 NumColorRows = ((NumGridBlocks - 1) / 16) + 1;
+	float SizeY = (NumColorRows * (BlockSize.Y + Padding)) - Padding;
+
+	float SizeX = (16 * (BlockSize.X + Padding)) - Padding;
+
+	return FVector2D(SizeX, SizeY);
 }
 
-/**
-* All widgets must provide a way to access their children in a layout-agnostic way.
-* Panels store their children in Slots, which creates a dilemma. Most panels
-* can store their children in a TPanelChildren<Slot>, where the Slot class
-* provides layout information about the child it stores. In that case
-* GetChildren should simply return the TPanelChildren<Slot>.
-*/
 FChildren* SThemeColorBlocksBar::GetChildren()
 {
-	return &ColorBlocks;
+	return &Children;
 }
 
 void SThemeColorBlocksBar::OnDragEnter( const FGeometry& MyGeometry, const FDragDropEvent& DragDropEvent )
@@ -578,9 +557,7 @@ void SThemeColorBlocksBar::OnDragEnter( const FGeometry& MyGeometry, const FDrag
 	TSharedPtr<FColorDragDrop> DragDropOperation = DragDropEvent.GetOperationAs<FColorDragDrop>();
 	if ( DragDropOperation.IsValid() )
 	{
-		NewColorPlaceholder = MakeShareable(new FLinearColor(DragDropOperation->Color));
-		NewColorBlockPlaceholder =
-			SNew(SThemeColorBlock) .Color(NewColorPlaceholder) .UseSRGB(bUseSRGB) .UseAlpha(bUseAlpha);
+		DragDropOperation->MarkForAdd();
 	}
 }
 
@@ -589,60 +566,61 @@ void SThemeColorBlocksBar::OnDragLeave( const FDragDropEvent& DragDropEvent )
 	TSharedPtr<FColorDragDrop> DragDropOperation = DragDropEvent.GetOperationAs<FColorDragDrop>();
 	if ( DragDropOperation.IsValid() )
 	{
-		DragDropOperation->SetDecoratorVisibility(true);
-
-		DestroyPlaceholders();
+		DragDropOperation->MarkForDelete();
+		PlaceholderIndex.Reset();
 	}
 }
 
-/**
-* Called during drag and drop when the the mouse is being dragged over a widget.
-*
-* @param MyGeometry      The geometry of the widget receiving the event.
-* @param DragDropEvent   The drag and drop event.
-*
-* @return A reply that indicated whether this event was handled.
-*/
 FReply SThemeColorBlocksBar::OnDragOver( const FGeometry& MyGeometry, const FDragDropEvent& DragDropEvent )
 {
 	TSharedPtr<FColorDragDrop> DragDropOperation = DragDropEvent.GetOperationAs<FColorDragDrop>();
 	if ( DragDropOperation.IsValid() )
 	{
-		const float ChildSizeX = MyGeometry.Size.X / (ColorBlocks.Num() + 1);
-		const float GrabOffsetX = PlaceholderInitialGrabOffset.X;
-		PlaceholderBlockOffset = MyGeometry.AbsoluteToLocal(DragDropEvent.GetScreenSpacePosition()).X -
-			((GrabOffsetX == 0.f) ? ChildSizeX * 0.5f : GrabOffsetX);
+		const FVector2D DragLocation = MyGeometry.AbsoluteToLocal(DragDropEvent.GetScreenSpacePosition());
 
-		DragDropOperation->SetDecoratorVisibility(false);
+		constexpr int32 BlockSizeWithPadding = 26;
+		const int32 GridX = DragLocation.X / BlockSizeWithPadding;
+		const int32 GridY = DragLocation.Y / BlockSizeWithPadding;
+
+		if ((GridY == 0) && (GridX < 2))
+		{
+			// The dragged block is over the combo button
+			DragDropOperation->MarkForAdd();
+			PlaceholderIndex = 3;
+		}
+		else if ((GridY == 0) && (GridX == 2))
+		{
+			// The dragged block is over the delete button
+			DragDropOperation->MarkForDelete();
+			PlaceholderIndex.Reset();
+		}
+		else
+		{
+			DragDropOperation->MarkForAdd();
+			const int32 NewPlaceholderIndex = (GridY * 16) + GridX;
+			PlaceholderIndex = FMath::Clamp(NewPlaceholderIndex, 3, ColorBlocks.Num() + 3);
+		}
 			
 		return FReply::Handled();
 	}
-	else
-	{
-		return FReply::Unhandled();
-	}
+
+	return FReply::Unhandled();
 }
 
-/**
-* Called when the user is dropping something onto a widget; terminates drag and drop.
-*
-* @param MyGeometry      The geometry of the widget receiving the event.
-* @param DragDropEvent   The drag and drop event.
-*
-* @return A reply that indicated whether this event was handled.
-*/
 FReply SThemeColorBlocksBar::OnDrop( const FGeometry& MyGeometry, const FDragDropEvent& DragDropEvent )
 {
 	TSharedPtr<FColorDragDrop> DragDropOperation = DragDropEvent.GetOperationAs<FColorDragDrop>();
 	if ( DragDropOperation.IsValid() )
 	{
-		const float ChildSizeX = MyGeometry.Size.X / (ColorBlocks.Num() + 1);
-		const float CurrentDragCenter = PlaceholderBlockOffset + ChildSizeX * 0.5f;
-		const int32 LocID = FMath::Clamp(int32(CurrentDragCenter / ChildSizeX), 0, ColorBlocks.Num());
-
-		AddNewColorBlock(DragDropOperation->Color, LocID);
+		// The combo button takes up two blocks, and the add/delete button takes up the third
+		if (PlaceholderIndex.IsSet() && PlaceholderIndex.GetValue() >= 3)
+		{
+			const int32 AddIndex = PlaceholderIndex.GetValue() - 3;
+			AddNewColorBlock(DragDropOperation->Color, AddIndex, true);
+		}
 
-		DragDropOperation->bSetForDeletion = true;
+		PlaceholderIndex.Reset();
+		DragDropOperation->MarkForDelete();
 
 		return FReply::Handled();
 	}
@@ -652,112 +630,216 @@ FReply SThemeColorBlocksBar::OnDrop( const FGeometry& MyGeometry, const FDragDro
 	}
 }
 
-void SThemeColorBlocksBar::AddNewColorBlock(FLinearColor Color, int32 InsertPosition)
+void SThemeColorBlocksBar::AddNewColorBlock(FLinearColor Color, int32 InsertPosition, bool bAllowRepeat)
 {
-	// Force the alpha component to 1 when we don't care about the alpha, it'll allow the color to work with alpha-important widgets later
-	if( !bUseAlpha.Get() )
+	// Do not add new colors to recents
+	if (ColorTheme == ThemesViewer->GetRecents())
 	{
-		Color.A = 1.0f;
+		return;
 	}
-	ColorTheme.Get()->InsertNewColor(MakeShareable(new FLinearColor(Color)), InsertPosition);
 
-	DestroyPlaceholders();
+	// Check if the color being added is the same as the last one added
+	bool bIsRepeatColor = false;
+	const TArray<TSharedPtr<FColorInfo>>& Colors = ColorTheme->GetColors();
+	if (Colors.Num() > 0)
+	{
+		TSharedPtr<FColorInfo> NewestColor = Colors[0];
+		if (NewestColor->Color->Equals(Color))
+		{
+			bIsRepeatColor = true;
+		}
+	}
 
-	SColorThemesViewer::SaveColorThemesToIni();
+	if (!bIsRepeatColor || bAllowRepeat)
+	{
+		ColorTheme->InsertNewColor(MakeShareable(new FLinearColor(Color)), InsertPosition);
+		SColorThemesViewer::SaveColorThemesToIni();
+	}
+}
+
+void SThemeColorBlocksBar::AddToRecents(FLinearColor Color)
+{
+	if (TSharedPtr<FColorTheme> Recents = ThemesViewer->GetRecents())
+	{
+		// When the recents theme is active, the first row will have 14 color blocks, and each subsequent row will have 16
+		constexpr int32 MaxNumRecentRows = 3;
+		constexpr int32 MaxNumRecentColors = ((MaxNumRecentRows - 1) * 16) + 14;
+
+		// If the recents theme is full, remove the oldest color before adding the new one
+		if (Recents->GetColors().Num() == MaxNumRecentColors)
+		{
+			TSharedPtr<FColorInfo> OldestColor = Recents->GetColors().Last();
+			Recents->RemoveColor(OldestColor->Color);
+		}
+
+		Recents->InsertNewColor(MakeShareable(new FLinearColor(Color)), 0);
+
+		SColorThemesViewer::SaveColorThemesToIni();
+	}
+}
+
+bool SThemeColorBlocksBar::IsRecentsThemeActive() const 
+{
+	return ThemesViewer->IsRecentsThemeActive();
 }
 
 int32 SThemeColorBlocksBar::RemoveColorBlock(TSharedPtr<FLinearColor> ColorToRemove)
 {
-	int32 Position = ColorTheme.Get()->RemoveColor(ColorToRemove);
+	int32 Position = ColorTheme->RemoveColor(ColorToRemove);
 	
 	SColorThemesViewer::SaveColorThemesToIni();
 
 	return Position;
 }
 
-void SThemeColorBlocksBar::DestroyPlaceholders()
+void SThemeColorBlocksBar::RemoveRefreshCallback()
 {
-	NewColorBlockPlaceholder.Reset();
-	NewColorPlaceholder.Reset();
-	PlaceholderBlockOffset = 0.f;
-	PlaceholderInitialGrabOffset = FVector2D(ForceInit);
+	// Deprecated function
 }
 
-void SThemeColorBlocksBar::RemoveRefreshCallback()
+void SThemeColorBlocksBar::AddRefreshCallback()
 {
-	ColorTheme.Get()->OnRefresh().Remove(RefreshCallbackHandle);
+	// Deprecated function
 }
 
-void SThemeColorBlocksBar::AddRefreshCallback()
+void SThemeColorBlocksBar::OnThemeChanged()
+{
+	// Remove the refresh callback from the old theme
+	ColorTheme->OnRefresh().Remove(RefreshCallbackHandle);
+
+	// Get the new active theme and add our refresh callback to it so we can update the UI when the theme colors change
+	ColorTheme = ThemesViewer->GetCurrentColorTheme();
+	RefreshCallbackHandle = ColorTheme->OnRefresh().Add(RefreshCallback);
+
+	Refresh();
+}
+
+FReply SThemeColorBlocksBar::OnAddButtonClicked()
+{
+	if (OnGetActiveColor.IsBound())
+	{
+		const FLinearColor ActiveColor = OnGetActiveColor.Execute();
+		AddNewColorBlock(ActiveColor, 0);
+	}
+	return FReply::Handled();
+}
+
+void SThemeColorBlocksBar::ShowDeleteButton()
+{
+	bShowDeleteButton = true;
+}
+
+void SThemeColorBlocksBar::HideDeleteButton()
 {
-	RefreshCallbackHandle = ColorTheme.Get()->OnRefresh().Add(RefreshCallback);
+	bShowDeleteButton = false;
 }
 
 void SThemeColorBlocksBar::Refresh()
 {
-	DestroyPlaceholders();
-	ColorBlocks.Empty();
+	Children.Empty();
+
+	Children.Add(ThemesViewer.ToSharedRef());
+	Children.Add(AddDeleteOverlay.ToSharedRef());
 
-	check(ColorTheme.Get().IsValid());
+	ColorBlocks.Empty();
+	check(ColorTheme.IsValid());
 
-	const TArray< TSharedPtr<FColorInfo> >& Theme = ColorTheme.Get()->GetColors();
-	for (int32 i = 0; i < Theme.Num(); ++i)
+	const TArray< TSharedPtr<FColorInfo> >& Theme = ColorTheme->GetColors();
+	for (int32 ColorIndex = 0; ColorIndex < Theme.Num(); ++ColorIndex)
 	{
+		const bool bSupportsDrag = (ColorTheme != ThemesViewer->GetRecents());
+
 		ColorBlocks.Add(
 			SNew(SThemeColorBlock)
-				.Color(Theme[i]->Color)
-				.ColorInfo(Theme[i])
+				.Color(Theme[ColorIndex]->Color)
+				.ColorInfo(Theme[ColorIndex])
 				.OnSelectColor(OnSelectColor)
 				.Parent(SharedThis(this))
-				.ShowTrashCallback(ShowTrashCallback)
-				.HideTrashCallback(HideTrashCallback)
+				.ShowTrashCallback(this, &SThemeColorBlocksBar::ShowDeleteButton)
+				.HideTrashCallback(this, &SThemeColorBlocksBar::HideDeleteButton)
 				.UseSRGB(bUseSRGB)
 				.UseAlpha(bUseAlpha)
+				.SupportsDrag(bSupportsDrag)
 			);
+
+		Children.Add(ColorBlocks.Last().ToSharedRef());
 	}
 }
 
 void SThemeColorBlocksBar::SetPlaceholderGrabOffset(FVector2D GrabOffset)
 {
-	PlaceholderInitialGrabOffset = GrabOffset;
+	// Deprecated function
+}
+
+EVisibility SThemeColorBlocksBar::GetAddButtonVisibility() const
+{
+	return bShowDeleteButton ? EVisibility::Hidden : EVisibility::Visible;
+}
+
+EVisibility SThemeColorBlocksBar::GetDeleteButtonVisibility() const
+{
+	return bShowDeleteButton ? EVisibility::Visible : EVisibility::Hidden;
 }
 
-/**
-* Most panels do not create widgets as part of their implementation, so
-* they do not need to implement a Construct()
-*/
 void SThemeColorBlocksBar::Construct(const FArguments& InArgs)
 {
-	ColorTheme = InArgs._ColorTheme;
 	OnSelectColor = InArgs._OnSelectColor;
-	ShowTrashCallback = InArgs._ShowTrashCallback;
-	HideTrashCallback = InArgs._HideTrashCallback;
+	OnGetActiveColor = InArgs._OnGetActiveColor;
 	bUseSRGB = InArgs._UseSRGB;
 	bUseAlpha = InArgs._UseAlpha;
 	
+	ThemesViewer = SNew(SColorThemesViewer);
+	ThemesViewer->OnCurrentThemeChanged().AddSP(this, &SThemeColorBlocksBar::OnThemeChanged);
+	
 	RefreshCallback = FSimpleDelegate::CreateSP(this, &SThemeColorBlocksBar::Refresh);
-	AddRefreshCallback();
 
-	DestroyPlaceholders();
+	ColorTheme = ThemesViewer->GetCurrentColorTheme();
+	RefreshCallbackHandle = ColorTheme->OnRefresh().Add(RefreshCallback);
 
-	if (!InArgs._EmptyText.IsEmpty())
-	{
-		EmptyHintTextBlock = SNew(SBorder)
-			.Padding( 8.0f , 4.0f , 0.0f , 0.0f )
-			.Content()
-			[
-				SNew(STextBlock)
-					.Font(FCoreStyle::GetDefaultFontStyle("Regular", 8))
-					.Text(InArgs._EmptyText)
-			];
-	}
-
-	Refresh();
-}
+	DragShadow = SNew(SBorder)
+		.BorderImage(FAppStyle::Get().GetBrush("ColorPicker.RoundedInputBorderHovered"))
+		.Padding(FMargin(1.0f, 1.0f))
+		[
+			SNew(SBox)
+				.WidthOverride(22.0f)
+				.HeightOverride(22.0f)
+				[
+					SNullWidget::NullWidget
+				]
+		];
 
+	AddDeleteOverlay = SNew(SOverlay)
+		+ SOverlay::Slot()
+		[
+			SNew(SButton)
+				.ButtonStyle(&FAppStyle::Get().GetWidgetStyle<FButtonStyle>("ColorPicker.AddButton"))
+				.ContentPadding(FMargin(4.0))
+				.Visibility(this, &SThemeColorBlocksBar::GetAddButtonVisibility)
+				.ToolTipText(LOCTEXT("AddToThemeTooltip", "Add the currently selected color to the current color theme"))
+				.OnClicked(this, &SThemeColorBlocksBar::OnAddButtonClicked)
+				.Content()
+				[
+					SNew(SImage).Image(FAppStyle::Get().GetBrush("Icons.Plus"))
+				]
+		]
 
+		+ SOverlay::Slot()
+		[
+			SNew(SButton)
+				.ButtonStyle(&FAppStyle::Get().GetWidgetStyle<FButtonStyle>("ColorPicker.DeleteButton"))
+				.ContentPadding(FMargin(4.0))
+				.Visibility(this, &SThemeColorBlocksBar::GetDeleteButtonVisibility)
+				.ToolTipText(LOCTEXT("AddToThemeTooltip", "Delete this color from the current color theme"))
+				.Content()
+				[
+					SNew(SImage).Image(FAppStyle::Get().GetBrush("Icons.Delete"))
+				]
+		];
 
+	Refresh();
+}
 
+PRAGMA_DISABLE_DEPRECATION_WARNINGS
 void SColorThemeBar::Construct(const FArguments& InArgs)
 {
 	ColorTheme = InArgs._ColorTheme.Get();
@@ -790,7 +872,7 @@ void SColorThemeBar::Construct(const FArguments& InArgs)
 							.ColorTheme(InArgs._ColorTheme)
 							.ShowTrashCallback(ShowTrashCallback)
 							.HideTrashCallback(HideTrashCallback)
-							.EmptyText(NSLOCTEXT("ColorThemesViewer", "NoColorsText", "No Colors Added Yet"))
+							.EmptyText(LOCTEXT("NoColorsText", "No Colors Added Yet"))
 							.UseSRGB(bUseSRGB)
 							.UseAlpha(bUseAlpha)
 					]
@@ -814,9 +896,10 @@ FReply SColorThemeBar::OnMouseButtonDown( const FGeometry& MyGeometry, const FPo
 
 	return FReply::Unhandled();
 }
-
+PRAGMA_ENABLE_DEPRECATION_WARNINGS
 
 TArray< TSharedPtr<FColorTheme> > SColorThemesViewer::ColorThemes;
+TSharedPtr<FColorTheme> SColorThemesViewer::Recents;
 TWeakPtr<FColorTheme> SColorThemesViewer::CurrentlySelectedThemePtr;
 bool SColorThemesViewer::bSRGBEnabled = true;
 
@@ -824,188 +907,222 @@ bool SColorThemesViewer::bSRGBEnabled = true;
 BEGIN_SLATE_FUNCTION_BUILD_OPTIMIZATION
 void SColorThemesViewer::Construct(const FArguments& InArgs)
 {
-	#define LOCTEXT_NAMESPACE "ColorThemesViewer"
-
-	bUseAlpha = InArgs._UseAlpha;
+	if (!Recents.IsValid())
+	{
+		Recents = MakeShared<FColorTheme>();
+	}
 
 	LoadColorThemesFromIni();
 
-	const FSlateFontInfo SmallLayoutFont = FAppStyle::Get().GetFontStyle("Regular");
-
-
-	// different menus that could be visible for the color themes menu
-	// standard menu with "new", "rename" and "delete"
-	MenuStandard =
-		SNew(SHorizontalBox)
-		+SHorizontalBox::Slot() .AutoWidth().Padding(3)
-		[
-			SNew(SButton) .Text(LOCTEXT("NewButton", "New")) .HAlign(HAlign_Center)
-				.OnClicked(this, &SColorThemesViewer::NewColorTheme)
-		]
-		+SHorizontalBox::Slot() .AutoWidth().Padding(3)
-		[
-			SNew(SButton) .Text(LOCTEXT("DuplicateButton", "Duplicate")) .HAlign(HAlign_Center)
-				.OnClicked(this, &SColorThemesViewer::DuplicateColorTheme)
-		]
-		+SHorizontalBox::Slot() .AutoWidth().Padding(3)
-		[
-			SNew(SButton) .Text(LOCTEXT("RenameButton", "Rename")) .HAlign(HAlign_Center)
-				.OnClicked(this, &SColorThemesViewer::MenuToRename)
-		]
-		+SHorizontalBox::Slot() .AutoWidth().Padding(3)
-		[
-			SNew(SButton) .Text(LOCTEXT("DeleteButton", "Delete")) .HAlign(HAlign_Center)
-				.OnClicked(this, &SColorThemesViewer::MenuToDelete)
-		];
+	if (!CurrentlySelectedThemePtr.IsValid())
+	{
+		CurrentlySelectedThemePtr = Recents;
+	}
 
-	// menu for renaming the currently selected color theme
-	MenuRename =
-		SNew(SHorizontalBox)
-		+SHorizontalBox::Slot() .FillWidth(1) .Padding(3)
-		[
-			SAssignNew(RenameTextBox, SEditableTextBox) .Font(SmallLayoutFont)
-				.OnTextChanged(this, &SColorThemesViewer::ChangeThemeName)
-				.OnTextCommitted(this, &SColorThemesViewer::CommitThemeName)
-		]
-		+SHorizontalBox::Slot().AutoWidth() .Padding(3)
-		[
-			SNew(SButton) .Text(LOCTEXT("AcceptRenameButton", "Accept")) .HAlign(HAlign_Right)
-				.OnClicked(this, &SColorThemesViewer::AcceptThemeName)
-				.IsEnabled(this, &SColorThemesViewer::CanAcceptThemeName)
-		]
-		+SHorizontalBox::Slot().AutoWidth() .Padding(3)
-		[
-			SNew(SButton) .Text(LOCTEXT("CancelRenameButton", "Cancel")) .HAlign(HAlign_Right)
-				.OnClicked(this, &SColorThemesViewer::MenuToStandard)
-		];
+	SAssignNew(RenameTextBox, SEditableTextBox)
+		.IsEnabled(false)
+		.OnTextCommitted(this, &SColorThemesViewer::CommitThemeName)
+		.ClearKeyboardFocusOnCommit(false)
+		.SelectAllTextWhenFocused(true);
 
-	// menu for confirming whether you wish to delete the currently selected color theme
-	MenuConfirmDelete =
-		SNew(SHorizontalBox)
-		+SHorizontalBox::Slot() .FillWidth(1) .HAlign(HAlign_Left) .Padding(3)
-		[
-			SNew(SButton) .Text(LOCTEXT("ConfirmDeleteButton", "Delete")) .HAlign(HAlign_Center)
-				.OnClicked(this, &SColorThemesViewer::DeleteColorTheme)
-		]
-		+SHorizontalBox::Slot() .FillWidth(1) .HAlign(HAlign_Center) .VAlign(VAlign_Center) .Padding(3)
-		[
-			SNew(STextBlock) .Text(LOCTEXT("ConfirmDeletePrompt", "Confirm Delete")) .Font(SmallLayoutFont)
-		]
-		+SHorizontalBox::Slot() .FillWidth(1) .HAlign(HAlign_Right) .Padding(3)
+	ChildSlot
 		[
-			SNew(SButton) .Text(LOCTEXT("CancelDeleteButton", "Cancel")) .HAlign(HAlign_Center)
-				.OnClicked(this, &SColorThemesViewer::MenuToStandard)
-		];
+			SNew(SComboButton)
+				.ComboButtonStyle(&FAppStyle::Get().GetWidgetStyle<FComboButtonStyle>("ColorPicker.ThemesComboButton"))
+				.OnMenuOpenChanged(this, &SColorThemesViewer::OnMenuOpenChanged)
+				.ToolTipText(LOCTEXT("ColorThemeComboButtonToolTip", "Color Theme Options"))
+				.VAlign(VAlign_Center)
+				.MenuContent()
+				[
+					BuildMenu()
+				]
+				.ButtonContent()
+				[
+					SNew(SHorizontalBox)
 
-	// menu for dropping colors into the trash
-	MenuTrashColor =
-		SNew(SVerticalBox)
-		+SVerticalBox::Slot().AutoHeight()
-		[
-			SNew(SColorTrash)
+					+ SHorizontalBox::Slot()
+						.MinWidth(16)
+						.MaxWidth(16)
+						[
+							SNew(SImage).Image(this, &SColorThemesViewer::GetComboButtonImage)
+						]
+				]
 		];
-
-	this->ChildSlot
-	[
-		SNew(SVerticalBox)
-		+SVerticalBox::Slot()
-			.AutoHeight()
-			.MaxHeight(455.0f) // Roughly matches the size of the color picker at max height
-		[
-			SNew(SBox)
-			.WidthOverride(360)
-			[
-				SAssignNew(ColorThemeList, SListView< TSharedPtr<FColorTheme> >)
-					.ListItemsSource(&ColorThemes)
-					.OnGenerateRow(this, &SColorThemesViewer::OnGenerateColorThemeBars)
-			]
-		]
-		+SVerticalBox::Slot()
-			.AutoHeight()
-			.Padding(0,15,0,0)
-		[
-			SAssignNew(Menu, SBorder)
-			.BorderImage(FAppStyle::Get().GetBrush("Brushes.Panel"))
-		]
-		+ SVerticalBox::Slot()
-			.AutoHeight()
-		[
-			SAssignNew(ErrorText, SErrorText)
-				.Visibility(this, &SColorThemesViewer::OnGetErrorTextVisibility)
-		]
-	];
-
-	MenuToStandardNoReturn();
-
-	if (!CurrentlySelectedThemePtr.IsValid())
-	{
-		CurrentlySelectedThemePtr = ColorThemes[0];
-	}
-	ColorThemeList->ClearSelection();
-	ColorThemeList->SetItemSelection(CurrentlySelectedThemePtr.Pin(), true);
-
-	#undef LOCTEXT_NAMESPACE
 }
 END_SLATE_FUNCTION_BUILD_OPTIMIZATION
 
 void SColorThemesViewer::MenuToStandardNoReturn()
 {
-	ErrorText->SetError(FString());
-	Menu->SetContent(MenuStandard.ToSharedRef());
+	// Deprecated function
 }
 
-FReply SColorThemesViewer::MenuToStandard()
+TSharedPtr<FColorTheme> SColorThemesViewer::GetCurrentColorTheme() const
 {
-	MenuToStandardNoReturn();
-	return FReply::Handled();
+	return CurrentlySelectedThemePtr.IsValid() ? CurrentlySelectedThemePtr.Pin() : ColorThemes[0];
 }
 
-FReply SColorThemesViewer::MenuToRename()
+TSharedPtr<FColorTheme> SColorThemesViewer::GetRecents() const
 {
-	RenameTextBox->SetText( FText::FromString( GetCurrentColorTheme()->Name ) );
-	Menu->SetContent(MenuRename.ToSharedRef());
-	return FReply::Handled();
+	return Recents;
 }
 
-FReply SColorThemesViewer::MenuToDelete()
+const FSlateBrush* SColorThemesViewer::GetComboButtonImage() const
 {
-	Menu->SetContent(MenuConfirmDelete.ToSharedRef());
-	return FReply::Handled();
+	return IsRecentsThemeActive() ? FAppStyle::Get().GetBrush("Icons.Recent") : FAppStyle::Get().GetBrush("ColorPicker.ColorThemesSmall");
+};
+
+bool SColorThemesViewer::IsRecentsThemeActive() const
+{
+	return (CurrentlySelectedThemePtr == Recents);
 }
 
-void SColorThemesViewer::MenuToTrash()
+void SColorThemesViewer::OnMenuOpenChanged(bool bIsOpen)
 {
-	Menu->SetContent(MenuTrashColor.ToSharedRef());
+	if (!bIsOpen)
+	{
+		StopRename();
+	}
 }
 
-void SColorThemesViewer::RefreshThemes()
+void SColorThemesViewer::PopulateThemesMenu(UToolMenu* InMenu)
 {
-	ColorThemeList->RequestListRefresh();
+	FToolMenuSection& Section = InMenu->FindOrAddSection("SavedThemes", LOCTEXT("SavedThemes", "Saved Color Themes"));
 
-	SaveColorThemesToIni();
+	for (const TSharedPtr<FColorTheme>& ColorTheme : ColorThemes)
+	{
+		FToolMenuEntry Entry = FToolMenuEntry::InitMenuEntry(
+			FName(ColorTheme->Name),
+			FText::FromString(ColorTheme->Name),
+			FText::FromString(ColorTheme->Name),
+			FSlateIcon(FAppStyle::GetAppStyleSetName(), "ColorPicker.ColorThemesSmall"),
+			FUIAction(
+				FExecuteAction::CreateSP(this, &SColorThemesViewer::SetCurrentColorTheme, ColorTheme),
+				FCanExecuteAction(),
+				FIsActionChecked::CreateLambda([ColorTheme] { return (CurrentlySelectedThemePtr == ColorTheme); })
+			),
+			EUserInterfaceActionType::RadioButton
+		);
+
+		if (CurrentlySelectedThemePtr == ColorTheme && RenameTextBox->IsEnabled())
+		{
+			RenameTextBox->SetText(FText::FromString(ColorTheme->Name));
+			Entry.MakeCustomWidget.BindLambda([this](const FToolMenuContext&, const FToolMenuCustomWidgetContext&) { return RenameTextBox.ToSharedRef(); });
+		}
 
-	MenuToStandardNoReturn();
+		Section.AddEntry(Entry);
+	}
+
+	Section.AddSeparator("EndSavedThemes");
 }
 
-TSharedPtr<FColorTheme> SColorThemesViewer::GetCurrentColorTheme() const
+TSharedRef<SWidget> SColorThemesViewer::BuildMenu()
 {
-	return CurrentlySelectedThemePtr.IsValid() ? CurrentlySelectedThemePtr.Pin() : ColorThemes[0];
+	UToolMenus* ToolsMenu = UToolMenus::Get();
+
+	if (!ToolsMenu->IsMenuRegistered("ColorThemeContextMenu"))
+	{
+		ToolsMenu->RegisterMenu("ColorThemeContextMenu", NAME_None, EMultiBoxType::Menu);
+	}
+
+	UToolMenu* ColorThemeMenu = ToolsMenu->ExtendMenu("ColorThemeContextMenu");
+	ColorThemeMenu->bShouldCloseWindowAfterMenuSelection = false;
+
+	FToolMenuSection& RecentsSection = ColorThemeMenu->FindOrAddSection("RecentsSection");
+
+	RecentsSection.AddMenuEntry(
+		"RecentsTheme",
+		LOCTEXT("RecentsTheme", "Recents"),
+		LOCTEXT("RecentsThemeToolTip", "Recently used colors"),
+		FSlateIcon(FAppStyle::GetAppStyleSetName(), "Icons.Recent"),
+		FUIAction(
+			FExecuteAction::CreateSP(this, &SColorThemesViewer::SetCurrentColorTheme, Recents),
+			FCanExecuteAction(),
+			FIsActionChecked::CreateSP(this, &SColorThemesViewer::IsRecentsThemeActive)
+		),
+		EUserInterfaceActionType::RadioButton
+	);
+
+	FToolMenuSection& DefaultSection = ColorThemeMenu->AddDynamicSection("SavedThemes", FNewToolMenuDelegate::CreateSP(this, &SColorThemesViewer::PopulateThemesMenu));
+
+	FToolMenuSection& AddThemeSection = ColorThemeMenu->FindOrAddSection("AddThemeSection");
+
+	AddThemeSection.AddMenuEntry(
+		"NewThemeButton",
+		LOCTEXT("CreateNewTheme", "Create New Theme"),
+		LOCTEXT("CreateNewThemeTooltip", "Create New Theme"),
+		FSlateIcon(FAppStyle::GetAppStyleSetName(), "Icons.PlusCircle"),
+		FUIAction(
+			FExecuteAction::CreateSP(this, &SColorThemesViewer::NewColorTheme),
+			FCanExecuteAction()
+		),
+		EUserInterfaceActionType::CollapsedButton
+	);
+
+	FToolMenuSection& EditCurrentThemeSection = ColorThemeMenu->FindOrAddSection("EditCurrentThemeSection", LOCTEXT("EditThemeSection", "Edit Current Theme"));
+
+	EditCurrentThemeSection.AddMenuEntry(
+		"RenameThemeButton",
+		LOCTEXT("RenameTheme", "Rename"),
+		LOCTEXT("RenameThemeToolTip", "Rename the currently selected color theme"),
+		FSlateIcon(FAppStyle::GetAppStyleSetName(), "GenericCommands.Rename"),
+		FUIAction(
+			FExecuteAction::CreateSP(this, &SColorThemesViewer::StartRename),
+			FCanExecuteAction::CreateLambda([this]() { return !IsRecentsThemeActive(); })
+		),
+		EUserInterfaceActionType::CollapsedButton
+	);
+
+	EditCurrentThemeSection.AddMenuEntry(
+		"DuplicateThemeButton",
+		LOCTEXT("DuplicateTheme", "Duplicate"),
+		LOCTEXT("DuplicateThemeTooltip", "Duplicate the currently selected color theme"),
+		FSlateIcon(FAppStyle::GetAppStyleSetName(), "Icons.Duplicate"),
+		FUIAction(
+			FExecuteAction::CreateSP(this, &SColorThemesViewer::DuplicateColorTheme),
+			FCanExecuteAction::CreateLambda([this]() { return !IsRecentsThemeActive(); })
+		),
+		EUserInterfaceActionType::CollapsedButton
+	);
+
+	EditCurrentThemeSection.AddMenuEntry(
+		"DeleteThemeButton",
+		LOCTEXT("DeleteTheme", "Delete"),
+		LOCTEXT("DeleteThemeTooltip", "Delete the currently selected color theme"),
+		FSlateIcon(FAppStyle::GetAppStyleSetName(), "Icons.Delete"),
+		FUIAction(
+			FExecuteAction::CreateSP(this, &SColorThemesViewer::DeleteColorTheme),
+			FCanExecuteAction::CreateLambda([this]() { return !IsRecentsThemeActive(); })
+		),
+		EUserInterfaceActionType::CollapsedButton
+	);
+
+	return ToolsMenu->GenerateWidget("ColorThemeContextMenu", FToolMenuContext());
+}
+
+void SColorThemesViewer::StartRename()
+{
+	RenameTextBox->SetEnabled(true);
+	UToolMenus::Get()->RefreshMenuWidget("ColorThemeContextMenu");
+	FSlateApplication::Get().SetKeyboardFocus(RenameTextBox.ToSharedRef());
+}
+
+void SColorThemesViewer::StopRename()
+{
+	RenameTextBox->SetEnabled(false);
+	UToolMenus::Get()->RefreshMenuWidget("ColorThemeContextMenu");
 }
 
 void SColorThemesViewer::SetUseAlpha( const TAttribute<bool>& InUseAlpha )
-{ 
-	bUseAlpha = InUseAlpha;
+{
 }
 
 void SColorThemesViewer::SetCurrentColorTheme(TSharedPtr<FColorTheme> NewTheme)
 {
 	// Set the current theme, requires a preexisting theme to be passed in
 	CurrentlySelectedThemePtr = NewTheme;
-	ColorThemeList->ClearSelection();
-	ColorThemeList->SetItemSelection(NewTheme, true);
-
 	CurrentThemeChangedEvent.Broadcast();
-	MenuToStandardNoReturn();
+
+	StopRename();
 }
 
 TSharedPtr<FColorTheme> SColorThemesViewer::IsColorTheme(const FString& ThemeName)
@@ -1066,68 +1183,37 @@ TSharedPtr<FColorTheme> SColorThemesViewer::GetDefaultColorTheme(bool bCreateNew
 	return GetColorTheme( Name.ToString() );
 }
 
-FReply SColorThemesViewer::AcceptThemeName()
-{
-	UpdateThemeNameFromTextBox();
-	return FReply::Handled();
-}
-
 void SColorThemesViewer::CommitThemeName(const FText& InText, ETextCommit::Type InCommitType)
 {
 	if (InCommitType == ETextCommit::OnEnter)
 	{
 		UpdateThemeNameFromTextBox();
+		StopRename();
 	}
 }
 
-void SColorThemesViewer :: UpdateThemeNameFromTextBox()
+void SColorThemesViewer::UpdateThemeNameFromTextBox()
 {
 	// Update the theme name if it differs, ensuring it is still unique
 	const FString Name = RenameTextBox->GetText().ToString();
 	if (GetCurrentColorTheme()->Name != Name)
 	{
 		GetCurrentColorTheme()->Name = MakeUniqueThemeName(Name);
-		RefreshThemes();
-	}
-}
-
-bool SColorThemesViewer::CanAcceptThemeName() const
-{
-	return !ErrorText->HasError();
-}
-
-void SColorThemesViewer::ChangeThemeName(const FText& InText)
-{
-	ErrorText->SetError(FString());
-
-	const FString ThemeName = InText.ToString();
-	for (int32 ThemeIndex = 0; ThemeIndex < ColorThemes.Num(); ++ThemeIndex)
-	{
-		const TSharedPtr<FColorTheme>& ColorTheme = ColorThemes[ThemeIndex];
-		if (ColorTheme != GetCurrentColorTheme() && ColorTheme->Name == ThemeName)
-		{
-			FFormatNamedArguments Args;
-			Args.Add(TEXT("Name"), InText);
-			ErrorText->SetError(FText::Format(NSLOCTEXT("ColorThemesViewer", "VerifyTextDup", "A theme already exists with the name '{Name}'."), Args).ToString());
-			return;
-		}
+		SaveColorThemesToIni();
 	}
 }
 
-EVisibility SColorThemesViewer::OnGetErrorTextVisibility() const
-{
-	return !CanAcceptThemeName() ? EVisibility::Visible : EVisibility::Collapsed;
-}
-
-FReply SColorThemesViewer::NewColorTheme()
+void SColorThemesViewer::NewColorTheme()
 {
 	// Create a new, defaultly named theme and update the display
-	GetDefaultColorTheme(true);
-	RefreshThemes();
-	return FReply::Handled();
+	TSharedPtr<FColorTheme> NewTheme = GetDefaultColorTheme(true);
+	SaveColorThemesToIni();
+
+	SetCurrentColorTheme(NewTheme);
+	StartRename();
 }
 
-FReply SColorThemesViewer::DuplicateColorTheme()
+void SColorThemesViewer::DuplicateColorTheme()
 {
 	// Create a copy of the existing current color theme
 	TArray< TSharedPtr<FColorInfo> > NewColors;
@@ -1137,54 +1223,35 @@ FReply SColorThemesViewer::DuplicateColorTheme()
 		NewColors.Add(MakeShareable(new FColorInfo(CurrentColors[ColorIndex]->Color, CurrentColors[ColorIndex]->Label)));
 	}
 	const FText Name = NSLOCTEXT("ColorThemesViewer", "CopyThemeNameAppend", " Copy");
-	NewColorTheme( CurrentlySelectedThemePtr.Pin()->Name + Name.ToString(), NewColors );
-	RefreshThemes();
-	return FReply::Handled();
+	TSharedPtr<FColorTheme> NewTheme = NewColorTheme( CurrentlySelectedThemePtr.Pin()->Name + Name.ToString(), NewColors );
+	SaveColorThemesToIni();
+
+	SetCurrentColorTheme(NewTheme);
+	StartRename();
 }
 
-FReply SColorThemesViewer::DeleteColorTheme()
+void SColorThemesViewer::DeleteColorTheme()
 {
 	// Delete the current color theme
 	ColorThemes.Remove(GetCurrentColorTheme());
+
 	if (!ColorThemes.Num())
 	{
 	 	// Create the default if none exists
 		GetDefaultColorTheme();
 	}
-	SetCurrentColorTheme(ColorThemes[0]);
-	RefreshThemes();
-	return FReply::Handled();
-}
 
-bool SColorThemesViewer::OnReadUseSRGB() const
-{
-	return bSRGBEnabled;
-}
-
-bool SColorThemesViewer::OnReadUseAlpha() const
-{
-	return bUseAlpha.Get();
-}
+	SetCurrentColorTheme(Recents);
+	SaveColorThemesToIni();
 
-TSharedRef<ITableRow> SColorThemesViewer::OnGenerateColorThemeBars(TSharedPtr<FColorTheme> InItem, const TSharedRef<STableViewBase>& OwnerTable)
-{
-	return
-	SNew( STableRow< TSharedPtr<FColorTheme> >, OwnerTable )
-	[
-		SNew(SColorThemeBar)
-		. ColorTheme(InItem)
-		. OnCurrentThemeChanged(this, &SColorThemesViewer::SetCurrentColorTheme)
-		. ShowTrashCallback(this, &SColorThemesViewer::MenuToTrash)
-		. HideTrashCallback(this, &SColorThemesViewer::MenuToStandardNoReturn)
-		. UseSRGB(this, &SColorThemesViewer::OnReadUseSRGB)
-		. UseAlpha(this, &SColorThemesViewer::OnReadUseAlpha)
-	];
+	UToolMenus::Get()->RefreshMenuWidget("ColorThemeContextMenu");
 }
 
 void SColorThemesViewer::LoadColorThemesFromIni()
 {
 	if (FPaths::FileExists(GEditorPerProjectIni))
 	{
+		// Load Themes
 		bool bThemesRemaining = true;
 		int32 ThemeID = 0;
 		while (bThemesRemaining)
@@ -1192,7 +1259,7 @@ void SColorThemesViewer::LoadColorThemesFromIni()
 			const FString ThemeName = GConfig->GetStr(TEXT("ColorThemes"), *FString::Printf(TEXT("Theme%i"), ThemeID), GEditorPerProjectIni);
 			if (!ThemeName.IsEmpty())
 			{
-				TSharedPtr<FColorTheme> ColorTheme = GetColorTheme( ThemeName );
+				TSharedPtr<FColorTheme> ColorTheme = GetColorTheme(ThemeName);
 				check( ColorTheme.IsValid() );
 				bool bColorsRemaining = true;
 				int32 ColorID = 0;
@@ -1228,6 +1295,30 @@ void SColorThemesViewer::LoadColorThemesFromIni()
 				bThemesRemaining = false;
 			}
 		}
+
+		// Load Recents
+		bool bColorsRemaining = true;
+		int32 ColorID = 0;
+		while (bColorsRemaining)
+		{
+			const FString ColorString = GConfig->GetStr(TEXT("RecentColors"), *FString::Printf(TEXT("Color%i"), ColorID), GEditorPerProjectIni);
+			if (!ColorString.IsEmpty())
+			{
+				// Add the color if it hasn't already
+				FLinearColor Color;
+				Color.InitFromString(ColorString);
+				if (Recents->FindApproxColor(Color) == INDEX_NONE)
+				{
+					TSharedPtr<FColorInfo> NewColor = MakeShareable(new FColorInfo(MakeShareable(new FLinearColor(Color))));
+					Recents->InsertNewColor(NewColor, 0);
+				}
+				++ColorID;
+			}
+			else
+			{
+				bColorsRemaining = false;
+			}
+		}
 	}
 	
 	if (!ColorThemes.Num())
@@ -1241,7 +1332,9 @@ void SColorThemesViewer::SaveColorThemesToIni()
 {
 	if (FPaths::FileExists(GEditorPerProjectIni))
 	{
+		// Save Color Themes
 		GConfig->EmptySection(TEXT("ColorThemes"), GEditorPerProjectIni);
+
 		for (int32 ThemeIndex = 0; ThemeIndex < ColorThemes.Num(); ++ThemeIndex)
 		{
 			const TSharedPtr<FColorTheme>& Theme = ColorThemes[ThemeIndex];
@@ -1256,5 +1349,20 @@ void SColorThemesViewer::SaveColorThemesToIni()
 				GConfig->SetString(TEXT("ColorThemes"), *FString::Printf(TEXT("Theme%iLabel%i"), ThemeIndex, ColorIndex), *Label.ToString(), GEditorPerProjectIni);
 			}
 		}
+
+		// Save Recents
+		GConfig->EmptySection(TEXT("RecentColors"), GEditorPerProjectIni);
+
+		const TArray< TSharedPtr<FColorInfo> >& Colors = Recents->GetColors();
+		for (int32 ColorIndex = 0; ColorIndex < Colors.Num(); ++ColorIndex)
+		{
+			const TSharedPtr<FLinearColor>& Color = Colors[ColorIndex]->Color;
+			const FText& Label = Colors[ColorIndex]->Label;
+			GConfig->SetString(TEXT("RecentColors"), *FString::Printf(TEXT("Color%i"), ColorIndex), *Color->ToString(), GEditorPerProjectIni);
+		}
 	}
+
+	GConfig->Flush(false, GEditorPerProjectIni);
 }
+
+#undef LOCTEXT_NAMESPACE
\ No newline at end of file
diff --git a/Engine/Source/Runtime/AppFramework/Public/Widgets/Colors/SColorPicker.h b/Engine/Source/Runtime/AppFramework/Public/Widgets/Colors/SColorPicker.h
--- a/Engine/Source/Runtime/AppFramework/Public/Widgets/Colors/SColorPicker.h
+++ b/Engine/Source/Runtime/AppFramework/Public/Widgets/Colors/SColorPicker.h
@@ -7,6 +7,7 @@
 #include "Input/Reply.h"
 #include "Layout/Visibility.h"
 #include "Widgets/DeclarativeSyntaxSupport.h"
+#include "Widgets/Input/SButton.h"
 #include "Widgets/SWidget.h"
 #include "Widgets/SCompoundWidget.h"
 #include "Widgets/SWindow.h"
@@ -17,7 +18,7 @@ class SBorder;
 class SColorThemesViewer;
 class SComboButton;
 class SThemeColorBlocksBar;
-
+enum class EColorBlockAlphaDisplayMode : uint8;
 
 /** Called when the color picker cancel button is pressed */
 DECLARE_DELEGATE_OneParam(FOnColorPickerCancelled, FLinearColor);
@@ -227,7 +228,10 @@ protected:
 
 	APPFRAMEWORK_API void BeginAnimation(FLinearColor Start, FLinearColor End);
 
+	UE_DEPRECATED(5.6, "HideSmallTrash is deprecated. The delete button for theme colors is managed by the SThemeColorBlocksBar widget.")
 	APPFRAMEWORK_API void HideSmallTrash();
+
+	UE_DEPRECATED(5.6, "ShowSmallTrash is deprecated. The delete button for theme colors is managed by the SThemeColorBlocksBar widget.")
 	APPFRAMEWORK_API void ShowSmallTrash();
 
 	/** Cycles the color picker's mode. */
@@ -254,7 +258,7 @@ protected:
 	 *
 	 * @return The new color preview box.
 	 */
-	APPFRAMEWORK_API TSharedRef<SWidget> MakeColorPreviewBox() const;
+	APPFRAMEWORK_API TSharedRef<SWidget> MakeColorPreviewBox();
 
 private:
 
@@ -270,8 +274,8 @@ private:
 	// Callback for getting the array of colors for a color spin box gradient.
 	APPFRAMEWORK_API TArray<FLinearColor> GetGradientColors(EColorPickerChannels Channel) const;
 
-	// Callback for getting the visibility of the alpha channel portion in color blocks.
-	APPFRAMEWORK_API EVisibility HandleAlphaColorBlockVisibility() const;
+	// Callback for getting the alpha display mode of the color preview swatches.
+	APPFRAMEWORK_API EColorBlockAlphaDisplayMode HandleColorPreviewAlphaMode() const;
 
 	// Callback for clicking the Cancel button.
 	APPFRAMEWORK_API FReply HandleCancelButtonClicked();
@@ -285,6 +289,18 @@ private:
 	// Callback for getting the visibility of the given color picker mode.
 	APPFRAMEWORK_API EVisibility HandleColorPickerModeVisibility(EColorPickerModes Mode) const;
 
+	// Callback for setting the visibility of the themes panel.
+	APPFRAMEWORK_API FReply ToggleThemePanelVisibility();
+
+	// Callback for getting the icon to display on the Show/Hide themes panel button.
+	APPFRAMEWORK_API const FSlateBrush* HandleThemePanelButtonImageBrush() const;
+
+	// Callback for getting the visibility of the themes panel.
+	APPFRAMEWORK_API EVisibility HandleThemesPanelVisibility() const;
+
+	// Callback for getting the tooltip for the color theme panel
+	APPFRAMEWORK_API FText GetColorThemePanelToolTipText() const;
+
 	// Callback for getting the end color of a color slider.
 	APPFRAMEWORK_API FLinearColor HandleColorSliderEndColor(EColorPickerChannels Channel) const;
 
@@ -328,14 +344,23 @@ private:
 	APPFRAMEWORK_API void HandleInteractiveChangeEnd(float NewValue);
 
 	// Callback for clicking the new color preview block.
-	APPFRAMEWORK_API FReply HandleNewColorBlockMouseButtonDown(const FGeometry& MyGeometry, const FPointerEvent& MouseEvent, bool bCheckAlpha);
+	APPFRAMEWORK_API FReply HandleNewColorPreviewClicked();
+
+	// Callback for clicking the old color preview block.
+	APPFRAMEWORK_API FReply HandleOldColorPreviewClicked();
+
+	// Callback for getting the visibility of the old/new color buttons
+	APPFRAMEWORK_API EVisibility HandleColorPreviewButtonVisibility() const;
+
+	// Callback for getting the visibility of the new color button icon
+	APPFRAMEWORK_API EVisibility HandleNewColorButtonVisibility() const;
+
+	// Callback for getting the visibility of the old color button icon
+	APPFRAMEWORK_API EVisibility HandleOldColorButtonVisibility() const;
 
 	// Callback for clicking the OK button.
 	APPFRAMEWORK_API FReply HandleOkButtonClicked();
 
-	// Callback for clicking the old color preview block.
-	APPFRAMEWORK_API FReply HandleOldColorBlockMouseButtonDown(const FGeometry& MyGeometry, const FPointerEvent& MouseEvent, bool bCheckAlpha);
-
 	// Callback for checking whether sRGB colors should be rendered.
 	APPFRAMEWORK_API bool HandleColorPickerUseSRGB() const;
 
@@ -354,18 +379,6 @@ private:
 	// Callback for selecting a color in the color theme bar.
 	APPFRAMEWORK_API void HandleThemeBarColorSelected(FLinearColor NewValue);
 
-	// Callback for getting the theme bar's color theme.
-	APPFRAMEWORK_API TSharedPtr<class FColorTheme> HandleThemeBarColorTheme() const;
-
-	// Callback for getting the visibility of the theme bar hint text.
-	APPFRAMEWORK_API EVisibility HandleThemeBarHintVisibility() const;
-
-	// Callback for determining whether the theme bar should display the alpha channel.
-	APPFRAMEWORK_API bool HandleThemeBarUseAlpha() const;
-
-	// Callback for theme viewer changes.
-	APPFRAMEWORK_API void HandleThemesViewerThemeChanged();
-
 private:
 	
 	/** The color that is being targeted as a TAttribute */
@@ -401,6 +414,9 @@ private:
 	/** Whether or not the color uses Alpha or not */
 	TAttribute<bool> bUseAlpha;
 
+	/** Whether or not to display color as sRGB */
+	bool bUseSRGB = false;
+
 	/** Prevents immediate refresh for performance reasons. */
 	bool bOnlyRefreshOnMouseUp;
 
@@ -416,21 +432,15 @@ private:
 	/** The widget which holds the currently selected theme */
 	TSharedPtr<SThemeColorBlocksBar> CurrentThemeBar;
 
-	/** Widget which is either the button to show the color themes viewer, or to be a color trash */
-	TSharedPtr<SBorder> ColorThemeButtonOrSmallTrash;
+	/** Button overlaid on top of the new color preview swatch that adds the new color to the current color theme */
+	TSharedPtr<SButton> NewColorButton;
 
-	/** The button to show the color themes viewer */
-	TSharedPtr<SComboButton> ColorThemeComboButton;
-
-	/** The small color trash shown in place of the combo button */
-	TSharedPtr<SWidget> SmallTrash;
+	/** Button overlaid on top of the old color preview swatch that adds the old color to the current color theme */
+	TSharedPtr<SButton> OldColorButton;
 
 	/** Sets the display Gamma setting - used to correct colors sampled from the screen */
 	TAttribute<float> DisplayGamma;
 
-	/** Stores the original sRGB option if this color picker temporarily overrides the global option. */
-	TOptional<bool> OriginalSRGBOption;
-
 	/** True if this color picker is an inline color picker */
 	bool bColorPickerIsInlineVersion;
 
@@ -446,6 +456,9 @@ private:
 	/** The current display gamma used to correct colors picked from the display. */
 	bool bClampValue;
 
+	/** Whether or not the color themes panel is visible or collapsed */
+	bool bIsThemePanelVisible;
+
 private:
 
 	/** Invoked when a new value is selected on the color wheel */
@@ -465,11 +478,6 @@ private:
 
 	/** Allows a details view to own the color picker so refreshing another details view doesn't close it */
 	TWeakPtr<SWidget> OptionalOwningDetailsView;
-
-private:
-
-	/** A static pointer to the global color themes viewer */
-	static APPFRAMEWORK_API TWeakPtr<SColorThemesViewer> ColorThemesViewer;
 };
 
 
diff --git a/Engine/Source/Runtime/AppFramework/Public/Widgets/Colors/SColorThemes.h b/Engine/Source/Runtime/AppFramework/Public/Widgets/Colors/SColorThemes.h
--- a/Engine/Source/Runtime/AppFramework/Public/Widgets/Colors/SColorThemes.h
+++ b/Engine/Source/Runtime/AppFramework/Public/Widgets/Colors/SColorThemes.h
@@ -12,6 +12,7 @@
 #include "Layout/Children.h"
 #include "Widgets/SCompoundWidget.h"
 #include "Widgets/SPanel.h"
+#include "Widgets/SOverlay.h"
 #include "Widgets/Views/STableViewBase.h"
 #include "Widgets/Views/STableRow.h"
 #include "Widgets/Colors/SColorBlock.h"
@@ -21,6 +22,7 @@ class SEditableTextBox;
 class SErrorText;
 class STextBlock;
 class SThemeColorBlocksBar;
+class UToolMenu;
 
 struct FColorInfo
 {
@@ -90,7 +92,7 @@ private:
  * The SColorTrash is a multipurpose widget which allows FColorDragDrops
  * to be dropped on to to be 
  */
-class SColorTrash
+class UE_DEPRECATED(5.6, "SColorTrash is deprecated. SThemeColorBlocksBar now uses a simple SButton for its delete button.") SColorTrash
 	: public SCompoundWidget
 {
 public:
@@ -160,6 +162,7 @@ public:
 		, _HideTrashCallback()
 		, _UseSRGB()
 		, _UseAlpha()
+		, _SupportsDrag(true)
 	{ }
 
 		/** A pointer to the color this block uses */
@@ -176,6 +179,7 @@ public:
 		
 		/** Callback to pass down to the FColorDragDrop for it to show the trash */
 		SLATE_EVENT(FSimpleDelegate, ShowTrashCallback)
+		
 		/** Callback to pass down to the FColorDragDrop for it to hide the trash */
 		SLATE_EVENT(FSimpleDelegate, HideTrashCallback)
 		
@@ -185,6 +189,9 @@ public:
 		/** Whether the ability to pick the alpha value is enabled */
 		SLATE_ATTRIBUTE(bool, UseAlpha)
 
+		/** Whether the color block supports drag/drop operations */
+		SLATE_ATTRIBUTE(bool, SupportsDrag)
+
 	SLATE_END_ARGS()
 
 	/**
@@ -203,12 +210,8 @@ private:
 	virtual FReply OnMouseButtonUp( const FGeometry& MyGeometry, const FPointerEvent& MouseEvent ) override;
 	virtual FReply OnDragDetected( const FGeometry& MyGeometry, const FPointerEvent& MouseEvent ) override;
 
-	
 	FLinearColor GetColor() const;
 
-	FSlateColor HandleBorderColor() const;
-	const FSlateBrush* HandleBorderImage() const;
-
 	FText GetRedText() const;
 	FText GetGreenText() const;
 	FText GetBlueText() const;
@@ -239,213 +242,21 @@ private:
 	
 	/** Event called when this block is clicked */
 	FOnLinearColorValueChanged OnSelectColor;
-	
-	/** Callback to pass down to the FColorDragDrop for it to show the trash */
-	FSimpleDelegate ShowTrashCallback;
-	/** Callback to pass down to the FColorDragDrop for it to hide the trash */
-	FSimpleDelegate HideTrashCallback;
-	
-	/** Whether to use display sRGB color */
-	TAttribute<bool> bUseSRGB;
-
-	/** Whether or not the color uses Alpha or not */
-	TAttribute<bool> bUseAlpha;
-
-	float DistanceDragged;
-};
-
-
-/**
- * SThemeColorBlocksBars are panels for dragging and dropping SColorThemeBlocks onto and off of.
- */
-class SThemeColorBlocksBar
-	: public SPanel
-{
-public:
-
-	SLATE_BEGIN_ARGS(SThemeColorBlocksBar)
-		: _ColorTheme()
-		, _OnSelectColor()
-		, _ShowTrashCallback()
-		, _HideTrashCallback()
-		, _EmptyText()
-		, _UseSRGB()
-		, _UseAlpha()
-	{ }
-
-		/** A pointer to the color theme that this bar should display */
-		SLATE_ATTRIBUTE(TSharedPtr<FColorTheme>, ColorTheme)
-
-		/** Event called when a color block is clicked */
-		SLATE_EVENT(FOnLinearColorValueChanged, OnSelectColor)
-
-		/** Callback to pass down to the FColorDragDrop for it to show the trash */
-		SLATE_EVENT(FSimpleDelegate, ShowTrashCallback)
-		/** Callback to pass down to the FColorDragDrop for it to hide the trash */
-		SLATE_EVENT(FSimpleDelegate, HideTrashCallback)
-		
-		/** Specify what the bar should display when no colors are present */
-		SLATE_ARGUMENT(FText, EmptyText)
-
-		/** Whether to display sRGB color */
-		SLATE_ATTRIBUTE(bool, UseSRGB)
-
-		/** Whether the ability to pick the alpha value is enabled */
-		SLATE_ATTRIBUTE(bool, UseAlpha)
-
-	SLATE_END_ARGS()
-
-	SThemeColorBlocksBar();
-
-	/**
-	 * Most panels do not create widgets as part of their implementation, so
-	 * they do not need to implement a Construct()
-	 */
-	void Construct(const FArguments& InArgs);
-	
-	/**
-	 * Panels arrange their children in a space described by the AllottedGeometry parameter. The results of the arrangement
-	 * should be returned by appending a FArrangedWidget pair for every child widget.
-	 *
-	 * @param AllottedGeometry    The geometry allotted for this widget by its parent.
-	 * @param ArrangedChildren    The array to which to add the WidgetGeometries that represent the arranged children.
-	 */
-	virtual void OnArrangeChildren( const FGeometry& AllottedGeometry, FArrangedChildren& ArrangedChildren ) const override;
-
-	/**
-	 * A Panel's desired size in the space required to arrange of its children on the screen while respecting all of
-	 * the children's desired sizes and any layout-related options specified by the user. See StackPanel for an example.
-	 *
-	 * @return The desired size.
-	 */
-	virtual FVector2D ComputeDesiredSize(float) const override;
-
-	/**
-	 * All widgets must provide a way to access their children in a layout-agnostic way.
-	 * Panels store their children in Slots, which creates a dilemma. Most panels
-	 * can store their children in a TPanelChildren<Slot>, where the Slot class
-	 * provides layout information about the child it stores. In that case
-	 * GetChildren should simply return the TPanelChildren<Slot>.
-	 */
-	virtual FChildren* GetChildren() override;
-	
-	/**
-	 * Called during drag and drop when the drag enters a widget.
-	 *
-	 * @param MyGeometry      The geometry of the widget receiving the event.
-	 * @param DragDropEvent   The drag and drop event.
-	 * @return A reply that indicated whether this event was handled.
-	 */
-	virtual void OnDragEnter( const FGeometry& MyGeometry, const FDragDropEvent& DragDropEvent ) override;
-	
-	/**
-	 * Called during drag and drop when the drag leaves a widget.
-	 *
-	 * @param DragDropEvent   The drag and drop event.
-	 * @return A reply that indicated whether this event was handled.
-	 */
-	virtual void OnDragLeave( const FDragDropEvent& DragDropEvent ) override;
-
-	/**
-	 * Called during drag and drop when the the mouse is being dragged over a widget.
-	 *
-	 * @param MyGeometry      The geometry of the widget receiving the event.
-	 * @param DragDropEvent   The drag and drop event.
-	 * @return A reply that indicated whether this event was handled.
-	 */
-	virtual FReply OnDragOver( const FGeometry& MyGeometry, const FDragDropEvent& DragDropEvent ) override;
-
-	/**
-	 * Called when the user is dropping something onto a widget; terminates drag and drop.
-	 *
-	 * @param MyGeometry      The geometry of the widget receiving the event.
-	 * @param DragDropEvent   The drag and drop event.
-	 * @return A reply that indicated whether this event was handled.
-	 */
-	virtual FReply OnDrop( const FGeometry& MyGeometry, const FDragDropEvent& DragDropEvent ) override;
-
-	/**
-	 * Adds a new color block to the Bar
-	 * @param Color				The	color the new block will be
-	 * @param InsertPosition	The location the block will be at
-	 */
-	void AddNewColorBlock(FLinearColor Color, int32 InsertPosition);
-
-	/**
-	 * @param ColorToRemove		The pointer to the color that should be removed
-	 *
-	 * @return					The index of the removed color block, INDEX_NONE if it can't be found
-	 */
-	int32 RemoveColorBlock(TSharedPtr<FLinearColor> ColorToRemove);
-
-	/**
-	 * Refresh Callbacks are used by Color Themes to notify widgets to refresh
-	 * This removes a callback from the Color Theme that this bar has a pointer to
-	 */
-	void RemoveRefreshCallback();
-
-	/** Adds a callback to the Color Theme that this bar has a pointer to. */
-	void AddRefreshCallback();
-
-	/**
-	 * Rebuilds the entire bar, regenerating all the constituent color blocks
-	 */
-	void Refresh();
-
-	void SetPlaceholderGrabOffset(FVector2D GrabOffset);
-	
-private:
-
-	/** Destroys the placeholder block in this widget. */
-	void DestroyPlaceholders();
-
-	/** The children blocks of this panel */
-	TSlotlessChildren<SThemeColorBlock> ColorBlocks;
-
-	/**
-	 * A placeholder child, which exists as a "preview" to show what would happen when dropping the current
-	 * FColorDragDrop onto this widget. It "replaces" the FColorDragDrop widget with a preview SThemeColorBlock
-	 */
-	TSharedPtr<SThemeColorBlock> NewColorBlockPlaceholder;
-
-	/** The placeholder's color, also used to determine whether the placeholder is real or not (by IsValid) */
-	TSharedPtr<FLinearColor> NewColorPlaceholder;
-
-	/** Current x offset of the placeholder block */
-	float PlaceholderBlockOffset;
-
-	/** The initial grab offset when grabbing the placeholder */
-	FVector2D PlaceholderInitialGrabOffset;
-
-	/** A help text widget which appears when there are no children in this panel */
-	TSharedPtr<SWidget> EmptyHintTextBlock;
-
-	/**
-	 * The Color Theme that this SThemeColorBlockBar is displaying.
-	 * This is a TAttribute so it can re-get the theme when it changes rather than rely on a delegate to refresh it
-	 */
-	TAttribute< TSharedPtr<FColorTheme> > ColorTheme;
-	
-	/** Event called when a color block is clicked */
-	FOnLinearColorValueChanged OnSelectColor;
-	
-	/** Callback to pass to the Color Theme. Holds a handle to this bar's Refresh method */
-	FSimpleDelegate RefreshCallback;
-
-	/** Handle to the registered RefreshCallback delegate */
-	FDelegateHandle RefreshCallbackHandle;
 
 	/** Callback to pass down to the FColorDragDrop for it to show the trash */
 	FSimpleDelegate ShowTrashCallback;
 
 	/** Callback to pass down to the FColorDragDrop for it to hide the trash */
 	FSimpleDelegate HideTrashCallback;
-
+	
 	/** Whether to use display sRGB color */
 	TAttribute<bool> bUseSRGB;
 
 	/** Whether or not the color uses Alpha or not */
 	TAttribute<bool> bUseAlpha;
+
+	/** Whether the color block supports drag/drop operations */
+	TAttribute<bool> bSupportsDrag;
 };
 
 
@@ -456,7 +267,7 @@ DECLARE_DELEGATE_OneParam(FOnCurrentThemeChanged, TSharedPtr<FColorTheme>)
  * SColorThemeBars include a ThemeColorBlocksBar in addition to a label.
  * Clicking on one will select it and set the currently used color theme to it
  */
-class SColorThemeBar
+class UE_DEPRECATED(5.6, "SColorThemeBar is deprecated. SColorThemesViewer now displays a standard menu with a list of available themes.") SColorThemeBar
 	: public SCompoundWidget
 {
 public:
@@ -527,20 +338,16 @@ private:
 
 
 /**
- * There should only ever be a single SColorThemesViewer. It is the widget which
- * manages all color themes and displays in a nice list format.
+ * The widget that manages the dropdown menu in the Color Themes Panel
  */
-class SColorThemesViewer
-	: public SCompoundWidget
+class SColorThemesViewer : public SCompoundWidget
 {
 public:
 
 	SLATE_BEGIN_ARGS(SColorThemesViewer)
-		: _UseAlpha()
 	{ }
 
-		/** Whether the ability to pick the alpha value is enabled */
-		SLATE_ATTRIBUTE(bool, UseAlpha)
+		SLATE_ATTRIBUTE_DEPRECATED(bool, UseAlpha, 5.6, "UseAlpha is deprecated. Set the UseAlpha attribute of SThemeColorBlocksBar to control the opacity of theme colors.")
 
 	SLATE_END_ARGS()
 	
@@ -549,8 +356,11 @@ public:
 	/** Gets the current color theme */
 	TSharedPtr<FColorTheme> GetCurrentColorTheme() const;
 
-	/** Sets the UseAlpha attribute */
-	void SetUseAlpha( const TAttribute<bool>& InUseAlpha );
+	/** Gets the recents color theme */
+	TSharedPtr<FColorTheme> GetRecents() const;
+
+	UE_DEPRECATED(5.6, "SetUseAlpha is deprecated. Set the UseAlpha attribute of SThemeColorBlocksBar to control the opacity of theme colors.")
+	void SetUseAlpha(const TAttribute<bool>& InUseAlpha);
 
 	/** Load the color theme settings from the config */
 	static void LoadColorThemesFromIni();
@@ -565,40 +375,43 @@ public:
 		return CurrentThemeChangedEvent;
 	}
 
-	/** Whether to use display sRGB color */
+	UE_DEPRECATED(5.6, "bSRGBEnabled is deprecated. Set the UseSRGB attribute of SThemeColorBlocksBar to control the sRGB display of theme colors.")
 	static bool bSRGBEnabled;
 
+	UE_DEPRECATED(5.6, "MenuToStandardNoReturn is deprecated. SColorThemesViewer completely manages its own ToolsMenu.")
 	void MenuToStandardNoReturn();
 
+	/** Returns true if the currently chosen theme is the "Recents" theme */
+	bool IsRecentsThemeActive() const;
+
 private:
+	/** Generate the content for the dropdown menu */
+	TSharedRef<SWidget> BuildMenu();
+
+	/** Add any entry to the dropdown menu for each available theme */
+	void PopulateThemesMenu(UToolMenu* InMenu);
+
+	/** Initiate a rename interaction */
+	void StartRename();
+
+	/** Terminate a rename interaction */
+	void StopRename();
+
+	void NewColorTheme();
+	void DuplicateColorTheme();
+	void DeleteColorTheme();
 
-	FReply NewColorTheme();
-	FReply DuplicateColorTheme();
-	FReply DeleteColorTheme();
-	FReply AcceptThemeName();
 	/** Rename theme if the user has pressed Enter in the RenameTextBox */
 	void CommitThemeName(const FText& InText, ETextCommit::Type InCommitType);
 	void UpdateThemeNameFromTextBox();
 	
-	bool CanAcceptThemeName() const;
-	void ChangeThemeName(const FText& InText);
-	EVisibility OnGetErrorTextVisibility() const;
-
 	/** Sets the current color theme to the existing theme */
 	void SetCurrentColorTheme(TSharedPtr<FColorTheme> NewTheme);
 
-	FReply MenuToStandard();
-	FReply MenuToRename();
-	FReply MenuToDelete();
-	void MenuToTrash();
-	
-	bool OnReadUseSRGB() const;
-	bool OnReadUseAlpha() const;
-	
-	/** Refreshes the list, saves colors to ini, and sends the menu to standard */
-	void RefreshThemes();
+	/** Callback when the dropdown menu is opened/closed */
+	void OnMenuOpenChanged(bool bIsOpen);
 
-	TSharedRef<ITableRow> OnGenerateColorThemeBars(TSharedPtr<FColorTheme> InItem, const TSharedRef<STableViewBase>& OwnerTable);
+	const FSlateBrush* GetComboButtonImage() const;
 
 	/** Gets the default color theme, optionally creates it if not present */
 	static TSharedPtr<FColorTheme> GetDefaultColorTheme(bool bCreateNew = false);
@@ -610,20 +423,8 @@ private:
 	static FString MakeUniqueThemeName(const FString& ThemeName);
 	/** Creates a new theme, ensuring the name is unique */
 	static TSharedPtr<FColorTheme> NewColorTheme(const FString& ThemeName, const TArray< TSharedPtr<FColorInfo> >& ThemeColors = TArray< TSharedPtr<FColorInfo> >());
-	
-	/** A list of all the color themes for the list view to use */
-	TSharedPtr< SListView <TSharedPtr <FColorTheme> > > ColorThemeList;
-
-	/** The menu is a widget at the bottom of the SColorThemesViewer with variable content */
-	TSharedPtr<SBorder> Menu;
-	/** The menu alternates between a standard menu, a rename menu, a confirm delete menu, and a color trash */
-	TSharedPtr<SWidget> MenuStandard;
-	TSharedPtr<SWidget> MenuRename;
-	TSharedPtr<SWidget> MenuConfirmDelete;
-	TSharedPtr<SWidget> MenuTrashColor;
-	/** If any error occurs, it needs to be displayed here as SPopupErrorText won't work with nested combo buttons */
-	TSharedPtr<SErrorText> ErrorText;
 
+private:
 	/** The text box for renaming themes */
 	TSharedPtr<SEditableTextBox> RenameTextBox;
 
@@ -636,40 +437,195 @@ private:
 	/** A static holder of the color themes for the entire program */
 	static TArray<TSharedPtr<FColorTheme>> ColorThemes;
 
+	/** A static color theme comprised of the most recently accepted color values */
+	static TSharedPtr<FColorTheme> Recents;
+
 	/** A static pointer to the color theme that is currently selected for the entire program */
 	static TWeakPtr<FColorTheme> CurrentlySelectedThemePtr;
 };
 
 
 /**
- * This operation is a color which can be dragged and dropped between widgets.
- * Represents a SThemeColorBlock that is dragged around, and can be dropped into a color trash.
+ * A panel for displaying SColorBlocks in a FColorTheme
  */
-class FColorDragDrop
-	: public FDragDropOperation
+class SThemeColorBlocksBar : public SPanel
 {
 public:
 
-	DRAG_DROP_OPERATOR_TYPE(FColorDragDrop, FDragDropOperation)
+	/** Delegate for hooking up to an inline editable text block 'IsSelected' check. */
+	DECLARE_DELEGATE_RetVal(FLinearColor, FOnGetActiveColor);
+
+	SLATE_BEGIN_ARGS(SThemeColorBlocksBar)
+		: _OnSelectColor()
+		, _OnGetActiveColor()
+		, _UseSRGB()
+		, _UseAlpha()
+		{ }
+
+		/** A pointer to the color theme that this bar should display */
+		SLATE_ATTRIBUTE_DEPRECATED(TSharedPtr<FColorTheme>, ColorTheme, 5.6, "ColorTheme Attribute is deprecated. This widget owns a SColorThemesViewer which supplies the current color theme.")
+
+		/** Event called when a color block is clicked */
+		SLATE_EVENT(FOnLinearColorValueChanged, OnSelectColor)
+
+		/** Event called to get the current color in the color picker window */
+		SLATE_EVENT(FOnGetActiveColor, OnGetActiveColor)
+
+		UE_DEPRECATED(5.6, "ShowTrashCallback is deprecated. Visibility of the delete button is managed by this widget internally.")
+		SLATE_EVENT(FSimpleDelegate, ShowTrashCallback)
+
+		UE_DEPRECATED(5.6, "HideTrashCallback is deprecated. Visibility of the delete button is managed by this widget internally.")
+		SLATE_EVENT(FSimpleDelegate, HideTrashCallback)
+
+		/** Specify what the bar should display when no colors are present */
+		SLATE_ARGUMENT_DEPRECATED(FText, EmptyText, 5.6, "EmptyText is deprecated. This widget no longer displays any text when no themes have been created.")
+
+		/** Whether to display sRGB color */
+		SLATE_ATTRIBUTE(bool, UseSRGB)
+
+		/** Whether the ability to pick the alpha value is enabled */
+		SLATE_ATTRIBUTE(bool, UseAlpha)
+
+	SLATE_END_ARGS()
+
+	SThemeColorBlocksBar();
+
+	void Construct(const FArguments& InArgs);
+
+	// Begin SWidget overrides
+	virtual void OnArrangeChildren(const FGeometry& AllottedGeometry, FArrangedChildren& ArrangedChildren) const override;
+	virtual FVector2D ComputeDesiredSize(float) const override;
+	virtual FChildren* GetChildren() override;
+	virtual void OnDragEnter(const FGeometry& MyGeometry, const FDragDropEvent& DragDropEvent) override;
+	virtual void OnDragLeave(const FDragDropEvent& DragDropEvent) override;
+	virtual FReply OnDragOver(const FGeometry& MyGeometry, const FDragDropEvent& DragDropEvent) override;
+	virtual FReply OnDrop(const FGeometry& MyGeometry, const FDragDropEvent& DragDropEvent) override;
+	// End SWidget overrides
 
 	/**
-	 * Invoked when the drag and drop operation has ended.
-	 * 
-	 * @param bDropWasHandled   true when the drop was handled by some widget; false otherwise
-	 * @param MouseEvent        The mouse event which caused the on drop to be called.
+	 * Adds a new color block to the Bar
+	 * @param Color				The	color the new block will be
+	 * @param InsertPosition	The location the block will be at
 	 */
-	virtual void OnDrop( bool bDropWasHandled, const FPointerEvent& MouseEvent ) override;
+	void AddNewColorBlock(FLinearColor Color, int32 InsertPosition, bool bAllowRepeat = false);
+
+	/** Adds a new color block to the Recents color theme */
+	void AddToRecents(FLinearColor Color);
 
-	/** 
-	 * Called when the mouse was moved during a drag and drop operation
+	/** Returns true if the currently selected theme is Recents */
+	bool IsRecentsThemeActive() const;
+
+	/**
+	 * @param ColorToRemove		The pointer to the color that should be removed
 	 *
-	 * @param DragDropEvent    The event that describes this drag drop operation.
+	 * @return					The index of the removed color block, INDEX_NONE if it can't be found
+	 */
+	int32 RemoveColorBlock(TSharedPtr<FLinearColor> ColorToRemove);
+
+	UE_DEPRECATED(5.6, "RemoveRefreshCallback is deprecated. Color theme changes are handled by this widget internally")
+	void RemoveRefreshCallback();
+
+	UE_DEPRECATED(5.6, "RemoveRefreshCallback is deprecated. Color theme changes are handled by this widget internally")
+	void AddRefreshCallback();
+
+	/**
+	 * Rebuilds the entire bar, regenerating all the constituent color blocks
+	 */
+	void Refresh();
+
+	UE_DEPRECATED(5.6, "SetPlaceholderGrabOffset is deprecated. The Placeholder position is managed by this widget internally.")
+	void SetPlaceholderGrabOffset(FVector2D GrabOffset);
+
+	/** Make the delete button visible and hide the add button */
+	void ShowDeleteButton();
+
+	/** Hide the delete button and make the add button visible */
+	void HideDeleteButton();
+
+private:
+	/** Callback when the add button is clicked to add the current color picker color to the current theme */
+	FReply OnAddButtonClicked();
+
+	/** Callback when the user changes the currently selected theme, which prompts this widget to update the color blocks being displayed */
+	void OnThemeChanged();
+
+	/** Get the visibility of the add button (which is hidden if the delete button is visible) */
+	EVisibility GetAddButtonVisibility() const;
+
+	/** Get the visibility of the delete button (which is hidden unless the user is dragging a color block) */
+	EVisibility GetDeleteButtonVisibility() const;
+
+	/** The children blocks of this panel */
+	TSlotlessChildren<SWidget> Children;
+
+	/** Array of color blocks belonging to the currently selected color theme */
+	TArray<TSharedPtr<SThemeColorBlock>> ColorBlocks;
+
+	/**
+	 * The Color Theme that this SThemeColorBlockBar is displaying.
+	 * This is a TAttribute so it can re-get the theme when it changes rather than rely on a delegate to refresh it
 	 */
+	TSharedPtr<FColorTheme> ColorTheme;
+
+	/** Event called when a color block is clicked */
+	FOnLinearColorValueChanged OnSelectColor;
+
+	/** Event called to retrieve the current color from the color picker */
+	FOnGetActiveColor OnGetActiveColor;
+
+	/** Callback to pass to the Color Theme. Holds a handle to this bar's Refresh method */
+	FSimpleDelegate RefreshCallback;
+
+	/** Handle to the registered RefreshCallback delegate */
+	FDelegateHandle RefreshCallbackHandle;
+
+	/** Whether to use display sRGB color */
+	TAttribute<bool> bUseSRGB;
+
+	/** Whether or not the color uses Alpha or not */
+	TAttribute<bool> bUseAlpha;
+
+	/** Whether or not the delete button is currently visible */
+	bool bShowDeleteButton = false;
+
+	/** Placeholder widget to show where a dropped color block will be added to the color theme */
+	TSharedPtr<SBorder> DragShadow;
+
+	/** Overlay widget that shows either the add button or the delete button */
+	TSharedPtr<SOverlay> AddDeleteOverlay;
+
+	/** Custom combo button widget featuring a menu with the list of color themes and actions related to those themes */
+	TSharedPtr<SColorThemesViewer> ThemesViewer;
+
+	/** Index in the panel indicating where the drag shadow placeholder widget should be drawn */
+	TOptional<int32> PlaceholderIndex;
+};
+
+
+/**
+ * This operation is a color which can be dragged and dropped between widgets.
+ * Represents a SThemeColorBlock that is dragged around, and can be dropped into a color trash.
+ */
+class FColorDragDrop : public FDragDropOperation
+{
+public:
+
+	DRAG_DROP_OPERATOR_TYPE(FColorDragDrop, FDragDropOperation)
+
+	FColorDragDrop(FLinearColor InColor, bool bInUseSRGB, bool bInUseAlpha, FSimpleDelegate InTrashShowCallback, FSimpleDelegate InTrashHideCallback);
+
+	// Begin FDragDropOperation overrides
+	virtual void OnDrop( bool bDropWasHandled, const FPointerEvent& MouseEvent ) override;
 	virtual void OnDragged( const class FDragDropEvent& DragDropEvent ) override;
-	
-	/** Gets the widget that will serve as the decorator unless overridden. If you do not override, you will have no decorator */
 	virtual TSharedPtr<SWidget> GetDefaultDecorator() const override;
-	
+	// End FDragDropOperation overrides
+
+	/** Makes the decorator window fully opaque to visually indicate that the color will be added when dropped */
+	void MarkForAdd();
+
+	/** Makes the decorator window slightly transparent to visually indicate that the color will be deleted when dropped */
+	void MarkForDelete();
+
 	/**
 	 * Makes a new FColorDragDrop to hold on to
 	 * @param InColor			The color to be dragged and dropped
@@ -680,37 +636,34 @@ public:
 	 * @param Origin			The SThemeColorBlockBar that this operation is from
 	 * @param OriginPosition	The position in it's origin that it is from
 	 */
-	static TSharedRef<FColorDragDrop> New( FLinearColor InColor, bool bSRGB, bool bUseAlpha,
+	static TSharedRef<FColorDragDrop> New(FLinearColor InColor, bool bSRGB, bool bUseAlpha,
 		FSimpleDelegate TrashShowCallback = FSimpleDelegate(), FSimpleDelegate TrashHideCallback = FSimpleDelegate(),
-		TSharedPtr<SThemeColorBlocksBar> Origin = TSharedPtr<SThemeColorBlocksBar>(), int32 OriginPosition = 0 );
+		TSharedPtr<SThemeColorBlocksBar> Origin = TSharedPtr<SThemeColorBlocksBar>(), int32 OriginPosition = 0);
 
 	/** The color currently held onto by this drag drop operation */
 	FLinearColor Color;
 
-	/** Whether or not the color uses SRGB or not */
+	/** Whether or not the color uses sRGB */
 	bool bUseSRGB;
 
-	/** Whether or not the color uses Alpha or not */
+	/** Whether or not the color uses Alpha */
 	bool bUseAlpha;
 
-	/**
-	 * The origin is the SThemeColorBlockBar where this operation is from.
-	 * Upon dropping, this operation jumps back to it's origin if possible
-	 */
+	UE_DEPRECATED(5.6, "OriginBar is deprecated. This operation no longer needs direct access to the SThemeColorBlocksBar widget.")
 	TWeakPtr<SThemeColorBlocksBar> OriginBar;
 
-	/** The origin position */
+	UE_DEPRECATED(5.6, "OriginBarPosition is deprecated. This operation no longer needs to know its position in the SThemeColorBlocksBar widget.")
 	int32 OriginBarPosition;
 
-	/** Callback to show the trash when this widget is created */
+	/** Callback to show the delete button of the SThemeColorBlocksBar when this is dropped */
 	FSimpleDelegate ShowTrash;
 
-	/** Callback to hide the trash when this widget is dropped */
+	/** Callback to hide the delete button of the SThemeColorBlocksBar when this is dropped */
 	FSimpleDelegate HideTrash;
 
-	/** Flag which ensures that OnDrop will not replace this block in it's origin */
+	UE_DEPRECATED(5.6, "bSetForDeletion is deprecated. This operation no longer deletes any color blocks from the SThemeColorBlocksBar widget.")
 	bool bSetForDeletion;
 
-	/** The size of the drag and drop color block */
+	UE_DEPRECATED(5.6, "BlockSize is deprecated. This operation now uses the size of the cursor decorator window internally.")
 	FVector2D BlockSize;
 };
diff --git a/Engine/Source/Runtime/SlateCore/Private/Styling/StarshipCoreStyle.cpp b/Engine/Source/Runtime/SlateCore/Private/Styling/StarshipCoreStyle.cpp
--- a/Engine/Source/Runtime/SlateCore/Private/Styling/StarshipCoreStyle.cpp
+++ b/Engine/Source/Runtime/SlateCore/Private/Styling/StarshipCoreStyle.cpp
@@ -337,6 +337,7 @@ TSharedRef<ISlateStyle> FStarshipCoreStyle::Create()
 		Style->Set("Icons.Cross", new IMAGE_BRUSH_SVG("Starship/Common/plus", Icon16x16));
 		Style->Set("Icons.Plus", new IMAGE_BRUSH_SVG("Starship/Common/plus", Icon16x16));
 		Style->Set("Icons.Minus", new IMAGE_BRUSH_SVG("Starship/Common/minus", Icon16x16));
+		Style->Set("Icons.AddCircle", new IMAGE_BRUSH_SVG("Starship/Common/add-circle", Icon16x16));
 		Style->Set("Icons.PlusCircle", new IMAGE_BRUSH_SVG("Starship/Common/plus-circle", Icon16x16));
 		Style->Set("Icons.MinusCircle", new IMAGE_BRUSH_SVG("Starship/Common/minus-circle", Icon16x16));
 		Style->Set("Icons.X", new IMAGE_BRUSH_SVG("Starship/Common/close", Icon16x16));
@@ -1948,22 +1949,26 @@ void FStarshipCoreStyle::SetupColorPickerStyles(TSharedRef<FStyle>& Style)
 		Style->Set("ColorPicker.RoundedInputBorderActive", new FSlateRoundedBoxBrush(FStyleColors::Transparent, InputFocusRadius, FStyleColors::Primary, InputFocusThickness));
 		Style->Set("ColorPicker.RoundedInputBorderHovered", new FSlateRoundedBoxBrush(FStyleColors::Transparent, InputFocusRadius, FStyleColors::Hover, InputFocusThickness));
 		Style->Set("ColorPicker.MultipleValuesBackground", new FSlateRoundedBoxBrush(FStyleColors::Input, InputFocusRadius));
-		Style->Set("ColorPicker.AlphaBackground", new IMAGE_BRUSH("Starship/Common/Checker", Icon16x16, FLinearColor::White, ESlateBrushTileType::Both));
-		Style->Set("ColorPicker.EyeDropper", new IMAGE_BRUSH("Icons/eyedropper_16px", Icon16x16));
+		Style->Set("ColorPicker.RecessedBackground", new FSlateRoundedBoxBrush(FStyleColors::Recessed, InputFocusRadius * 2.0f));
+
 		Style->Set("ColorPicker.Font", FStyleFonts::Get().Normal);
 		Style->Set("ColorPicker.SmallFont", FStyleFonts::Get().Small);
+
+		Style->Set("ColorPicker.AlphaBackground", new IMAGE_BRUSH("Starship/Common/Checker", Icon16x16, FLinearColor::White, ESlateBrushTileType::Both));
+		Style->Set("ColorPicker.EyeDropper", new IMAGE_BRUSH("Icons/eyedropper_16px", Icon16x16));
 		Style->Set("ColorPicker.Mode", new IMAGE_BRUSH("Common/ColorPicker_Mode_16x", Icon16x16));
 		Style->Set("ColorPicker.ModeWheel", new IMAGE_BRUSH("Common/ColorWheel_20x", Icon20x20));
 		Style->Set("ColorPicker.ModeSpectrum", new IMAGE_BRUSH("Common/ColorSpectrum_20x", Icon20x20));
 		Style->Set("ColorPicker.Separator", new IMAGE_BRUSH("Common/ColorPicker_Separator", FVector2f(2.0f, 2.0f)));
 		Style->Set("ColorPicker.Selector", new IMAGE_BRUSH("Common/Circle", Icon8x8));
+		Style->Set("ColorPicker.SpinBoxSelectorVertical", new IMAGE_BRUSH("Starship/Common/color-grading-spinbox-selector", FVector2f(3.0f, 1.0f)));
+		Style->Set("ColorPicker.SpinBoxSelectorHorizontal", new IMAGE_BRUSH("Starship/Common/color-grading-spinbox-selector-v", FVector2f(1.0f, 3.0f)));
+
 		Style->Set("ColorPicker.Slider", FSliderStyle()
 			.SetDisabledThumbImage(IMAGE_BRUSH("Common/ColorPicker_SliderHandle", FVector2f(8.0f, 32.0f)))
 			.SetNormalThumbImage(IMAGE_BRUSH("Common/ColorPicker_SliderHandle", FVector2f(8.0f, 32.0f)))
 			.SetHoveredThumbImage(IMAGE_BRUSH("Common/ColorPicker_SliderHandle", FVector2f(8.0f, 32.0f)))
 		);
-		Style->Set("ColorPicker.SpinBoxSelectorVertical", new IMAGE_BRUSH("Starship/Common/color-grading-spinbox-selector", FVector2f(3.0f, 1.0f)));
-		Style->Set("ColorPicker.SpinBoxSelectorHorizontal", new IMAGE_BRUSH("Starship/Common/color-grading-spinbox-selector-v", FVector2f(1.0f, 3.0f)));
 	}
 
 	// SColorSlider defaults...
@@ -1984,7 +1989,39 @@ void FStarshipCoreStyle::SetupColorPickerStyles(TSharedRef<FStyle>& Style)
 
 	// SColorThemes defaults...
 	{
-		Style->Set("ColorThemes.DeleteButton", new IMAGE_BRUSH("Common/X", Icon16x16));
+		Style->Set("ColorPicker.ColorThemes", new IMAGE_BRUSH_SVG("Starship/Common/ColorThemes_20", Icon20x20));
+		Style->Set("ColorPicker.ColorThemesSmall", new IMAGE_BRUSH_SVG("Starship/Common/ColorThemes_16", Icon16x16));
+		Style->Set("ColorPicker.ColorThemesOff", new IMAGE_BRUSH_SVG("Starship/Common/ColorThemesOff_20", Icon20x20));
+
+		FWindowStyle CursorDecorator = Style->GetWidgetStyle<FWindowStyle>("Window");
+		CursorDecorator.SetCornerRadius(4);
+		Style->Set("ColorPicker.CursorDecorator", CursorDecorator);
+
+		FComboButtonStyle ThemesComboButton = Style->GetWidgetStyle<FComboButtonStyle>("ComboButton");
+		ThemesComboButton.SetDownArrowPadding(FMargin(0.0f));
+		ThemesComboButton.ButtonStyle
+			.SetNormal(FSlateRoundedBoxBrush(FStyleColors::Dropdown, InputFocusRadius, FStyleColors::Input, InputFocusThickness))
+			.SetDisabled(FSlateRoundedBoxBrush(FStyleColors::Dropdown, InputFocusRadius, FStyleColors::Input, InputFocusThickness))
+			.SetHovered(FSlateRoundedBoxBrush(FStyleColors::Dropdown, InputFocusRadius, FStyleColors::Input, InputFocusThickness))
+			.SetPressed(FSlateRoundedBoxBrush(FStyleColors::Dropdown, InputFocusRadius, FStyleColors::Input, InputFocusThickness));
+
+		Style->Set("ColorPicker.ThemesComboButton", ThemesComboButton);
+
+		FButtonStyle AddButton = Style->GetWidgetStyle<FButtonStyle>("Button");
+		AddButton.SetNormalPadding(0.0);
+		AddButton.SetPressedPadding(0.0);
+		AddButton.SetNormal(FSlateRoundedBoxBrush(FStyleColors::Transparent, InputFocusRadius, FStyleColors::DropdownOutline, InputFocusThickness));
+		AddButton.SetDisabled(FSlateRoundedBoxBrush(FStyleColors::Transparent, InputFocusRadius, FStyleColors::DropdownOutline, InputFocusThickness));
+		AddButton.SetHovered(FSlateRoundedBoxBrush(FStyleColors::Transparent, InputFocusRadius, FStyleColors::Hover2, InputFocusThickness));
+		AddButton.SetPressed(FSlateRoundedBoxBrush(FStyleColors::Transparent, InputFocusRadius, FStyleColors::Hover2, InputFocusThickness));
+
+		Style->Set("ColorPicker.AddButton", AddButton);
+
+		FButtonStyle DeleteButton = AddButton;
+		DeleteButton.SetNormal(FSlateRoundedBoxBrush(FStyleColors::Transparent, InputFocusRadius, FStyleColors::Transparent, InputFocusThickness));
+		DeleteButton.SetHovered(FSlateRoundedBoxBrush(FStyleColors::Transparent, InputFocusRadius, FStyleColors::DropdownOutline, InputFocusThickness));
+
+		Style->Set("ColorPicker.DeleteButton", DeleteButton);
 	}
 
 	// SColorWheel defaults...
@@ -1998,6 +2035,17 @@ void FStarshipCoreStyle::SetupColorPickerStyles(TSharedRef<FStyle>& Style)
 		Style->Set("ColorGradingWheel.HueValueCircle", new IMAGE_BRUSH("Common/ColorGradingWheel", FVector2f(192.f, 192.f)));
 		Style->Set("ColorGradingWheel.Selector", new IMAGE_BRUSH("Common/Circle", FVector2f(8.f, 8.f)));
 	}
+
+	// Color Preview Button defaults...
+	{
+		FButtonStyle ColorPreviewButton = Style->GetWidgetStyle<FButtonStyle>("Button");
+		ColorPreviewButton.SetNormal(FSlateRoundedBoxBrush(FStyleColors::Transparent, InputFocusRadius, FStyleColors::Transparent, InputFocusThickness));
+		ColorPreviewButton.SetHovered(FSlateRoundedBoxBrush(FLinearColor(0.0, 0.0, 0.0, 0.4), InputFocusRadius, FStyleColors::Transparent, InputFocusThickness));
+		ColorPreviewButton.SetPressed(FSlateRoundedBoxBrush(FLinearColor(0.0, 0.0, 0.0, 0.4), InputFocusRadius, FStyleColors::Transparent, InputFocusThickness));
+		ColorPreviewButton.SetDisabled(FSlateRoundedBoxBrush(FStyleColors::Transparent, InputFocusRadius, FStyleColors::Transparent, InputFocusThickness));
+
+		Style->Set("ColorPicker.ColorPreviewButton", ColorPreviewButton);
+	}
 }
 
 void FStarshipCoreStyle::SetupTableViewStyles(TSharedRef<FStyle>& Style)
