CIS: Fix module dependency in non-editor target

diff --git a/Engine/Source/Runtime/AppFramework/AppFramework.Build.cs b/Engine/Source/Runtime/AppFramework/AppFramework.Build.cs
--- a/Engine/Source/Runtime/AppFramework/AppFramework.Build.cs
+++ b/Engine/Source/Runtime/AppFramework/AppFramework.Build.cs
@@ -21,10 +21,18 @@ public class AppFramework : ModuleRules
 		PrivateDependencyModuleNames.AddRange(
 			new string[] {
 				"InputCore",
-				"ToolMenus",
 			}
 		);
 
+		if (Target.bBuildEditor == true)
+		{
+			PrivateDependencyModuleNames.AddRange(
+				new string[] {
+					"ToolMenus",
+				}
+			);
+		}
+
 		if (Target.Configuration != UnrealTargetConfiguration.Shipping)
 		{
 			PrivateIncludePathModuleNames.AddRange(
