ColorPicker: Fix Color Spectrum widget selector resetting to the center of widget if you drag beyond the bottom of the widget
#rnx

diff --git a/Engine/Source/Runtime/AppFramework/Private/Widgets/Colors/SColorPicker.cpp b/Engine/Source/Runtime/AppFramework/Private/Widgets/Colors/SColorPicker.cpp
--- a/Engine/Source/Runtime/AppFramework/Private/Widgets/Colors/SColorPicker.cpp
+++ b/Engine/Source/Runtime/AppFramework/Private/Widgets/Colors/SColorPicker.cpp
@@ -158,7 +158,7 @@ void SColorPicker::GenerateDefaultColorPickerContent( bool bAdvancedSectionExpan
 												SNew(SColorWheel)
 													.SelectedColor(this, &SColorPicker::GetCurrentColor)
 													.Visibility(this, &SColorPicker::HandleColorPickerModeVisibility, EColorPickerModes::Wheel)
-													.OnValueChanged(this, &SColorPicker::HandleColorSpectrumValueChanged)
+													.OnValueChanged(this, &SColorPicker::HandleColorWheelValueChanged)
 													.OnMouseCaptureBegin(this, &SColorPicker::HandleInteractiveChangeBegin)
 													.OnMouseCaptureEnd(this, &SColorPicker::HandleInteractiveChangeEnd)
 											]
@@ -535,7 +535,7 @@ void SColorPicker::GenerateInlineColorPickerContent()
 				SNew(SColorWheel)
 					.SelectedColor(this, &SColorPicker::GetCurrentColor)
 					.Visibility(this, &SColorPicker::HandleColorPickerModeVisibility, EColorPickerModes::Wheel)
-					.OnValueChanged(this, &SColorPicker::HandleColorSpectrumValueChanged)
+					.OnValueChanged(this, &SColorPicker::HandleColorWheelValueChanged)
 					.OnMouseCaptureBegin(this, &SColorPicker::HandleInteractiveChangeBegin)
 					.OnMouseCaptureEnd(this, &SColorPicker::HandleInteractiveChangeEnd)
 			]
@@ -1031,7 +1031,7 @@ FLinearColor SColorPicker::HandleColorSliderStartColor( EColorPickerChannels Cha
 }
 
 
-void SColorPicker::HandleColorSpectrumValueChanged( FLinearColor NewValue )
+void SColorPicker::HandleColorWheelValueChanged(FLinearColor NewValue)
 {
 	// In this color, R = H, G = S, B = V
 	if (FMath::IsNearlyZero(NewValue.B))
@@ -1046,6 +1046,12 @@ void SColorPicker::HandleColorSpectrumValueChanged( FLinearColor NewValue )
 }
 
 
+void SColorPicker::HandleColorSpectrumValueChanged(FLinearColor NewValue)
+{
+	SetNewTargetColorHSV(NewValue);
+}
+
+
 float SColorPicker::HandleColorSpinBoxValue( EColorPickerChannels Channel ) const
 {
 	switch (Channel)
diff --git a/Engine/Source/Runtime/AppFramework/Public/Widgets/Colors/SColorPicker.h b/Engine/Source/Runtime/AppFramework/Public/Widgets/Colors/SColorPicker.h
--- a/Engine/Source/Runtime/AppFramework/Public/Widgets/Colors/SColorPicker.h
+++ b/Engine/Source/Runtime/AppFramework/Public/Widgets/Colors/SColorPicker.h
@@ -307,6 +307,9 @@ private:
 	// Callback for getting the start color of a color slider.
 	APPFRAMEWORK_API FLinearColor HandleColorSliderStartColor(EColorPickerChannels Channel) const;
 
+	// Callback for value changes in the color wheel picker.
+	APPFRAMEWORK_API void HandleColorWheelValueChanged(FLinearColor NewValue);
+
 	// Callback for value changes in the color spectrum picker.
 	APPFRAMEWORK_API void HandleColorSpectrumValueChanged(FLinearColor NewValue);
 
