CIS: Fix localization conflicts

#rnx

diff --git a/Engine/Source/Runtime/AppFramework/Private/Widgets/Colors/SColorPicker.cpp b/Engine/Source/Runtime/AppFramework/Private/Widgets/Colors/SColorPicker.cpp
--- a/Engine/Source/Runtime/AppFramework/Private/Widgets/Colors/SColorPicker.cpp
+++ b/Engine/Source/Runtime/AppFramework/Private/Widgets/Colors/SColorPicker.cpp
@@ -823,7 +823,7 @@ TSharedRef<SWidget> SColorPicker::MakeColorPreviewBox() const
 						.ShowBackgroundForAlpha(true)
 						.AlphaDisplayMode(SharedThis(this), &SColorPicker::HandleColorPreviewAlphaMode)
 						.Color(this, &SColorPicker::GetCurrentColor)
-						.ToolTipText(LOCTEXT("NewColorButtonToolTip", "Preview of the currently selected color"))
+						.ToolTipText(LOCTEXT("NewColorBlockToolTip", "Preview of the currently selected color"))
 						.UseSRGB(SharedThis(this), &SColorPicker::HandleColorPickerUseSRGB)
 						.Size(FVector2D(106.0, 32.0))
 						.CornerRadius(FVector4(4.0, 4.0, 4.0, 4.0))
@@ -863,7 +863,7 @@ TSharedRef<SWidget> SColorPicker::MakeColorPreviewBox() const
 						.ShowBackgroundForAlpha(true)
 						.AlphaDisplayMode(SharedThis(this), &SColorPicker::HandleColorPreviewAlphaMode)
 						.Color(OldColor)
-						.ToolTipText(LOCTEXT("NewColorButtonToolTip", "Preview of the previously selected color"))
+						.ToolTipText(LOCTEXT("OldColorBlockToolTip", "Preview of the previously selected color"))
 						.UseSRGB(SharedThis(this), &SColorPicker::HandleColorPickerUseSRGB)
 						.Size(FVector2D(106.0, 32.0))
 						.CornerRadius(FVector4(4.0, 4.0, 4.0, 4.0))
diff --git a/Engine/Source/Runtime/AppFramework/Private/Widgets/Colors/SColorThemes.cpp b/Engine/Source/Runtime/AppFramework/Private/Widgets/Colors/SColorThemes.cpp
--- a/Engine/Source/Runtime/AppFramework/Private/Widgets/Colors/SColorThemes.cpp
+++ b/Engine/Source/Runtime/AppFramework/Private/Widgets/Colors/SColorThemes.cpp
@@ -829,7 +829,7 @@ void SThemeColorBlocksBar::Construct(const FArguments& InArgs)
 				.ButtonStyle(&FAppStyle::Get().GetWidgetStyle<FButtonStyle>("ColorPicker.DeleteButton"))
 				.ContentPadding(FMargin(4.0))
 				.Visibility(this, &SThemeColorBlocksBar::GetDeleteButtonVisibility)
-				.ToolTipText(LOCTEXT("AddToThemeTooltip", "Delete this color from the current color theme"))
+				.ToolTipText(LOCTEXT("DeleteFromThemeTooltip", "Delete this color from the current color theme"))
 				.Content()
 				[
 					SNew(SImage).Image(FAppStyle::Get().GetBrush("Icons.Delete"))
