ColorPicker: Fix default spawn location so the color picker window does not spawn partially underneath the taskbar when there are many color theme rows
#rnx

diff --git a/Engine/Source/Runtime/AppFramework/Private/Widgets/Colors/SColorPicker.cpp b/Engine/Source/Runtime/AppFramework/Private/Widgets/Colors/SColorPicker.cpp
--- a/Engine/Source/Runtime/AppFramework/Private/Widgets/Colors/SColorPicker.cpp
+++ b/Engine/Source/Runtime/AppFramework/Private/Widgets/Colors/SColorPicker.cpp
@@ -33,7 +33,7 @@
 
 
 /** A default window size for the color picker which looks nice */
-const FVector2D SColorPicker::DEFAULT_WINDOW_SIZE = FVector2D(462, 526);
+const FVector2D SColorPicker::DEFAULT_WINDOW_SIZE = FVector2D(462, 446);
 
 /** The max time allowed for updating before we shut off auto-updating */
 const double SColorPicker::MAX_ALLOWED_UPDATE_TIME = 0.1;
@@ -1386,12 +1386,14 @@ bool OpenColorPicker(const FColorPickerArgs& Args)
 	ensureMsgf(Args.OnColorCommitted.IsBound(), TEXT("OnColorCommitted should be bound to set the color."));
 		
 	// Determine the position of the window so that it will spawn near the mouse, but not go off the screen.
-	FVector2D CursorPos = FSlateApplication::Get().GetCursorPos();
-
-	FSlateRect Anchor(CursorPos.X, CursorPos.Y, CursorPos.X, CursorPos.Y);
-
-	FVector2D AdjustedSummonLocation = FSlateApplication::Get().CalculatePopupWindowPosition( Anchor, SColorPicker::DEFAULT_WINDOW_SIZE, true, FVector2D::ZeroVector, Orient_Horizontal );
-
+	const FVector2D CursorPos = FSlateApplication::Get().GetCursorPos();
+	const FSlateRect Anchor(CursorPos.X, CursorPos.Y, CursorPos.X, CursorPos.Y);
+
+	// Because the window has not yet been created, its desired size is still unknown. 
+	// This estimate is the size of the window with 4 rows of color theme blocks, which should be large enough in most cases to compute a reasonable summon location.
+	const FVector2D PaddingForColorTheme = FVector2D(0, 130);
+	const FVector2D WindowSizeEstimate = SColorPicker::DEFAULT_WINDOW_SIZE + PaddingForColorTheme;
+	const FVector2D AdjustedSummonLocation = FSlateApplication::Get().CalculatePopupWindowPosition( Anchor, WindowSizeEstimate, true, FVector2D::ZeroVector, Orient_Horizontal );
 
 	// Only override the color picker window creation behavior if we are not creating a modal color picker
 	const bool bOverrideNonModalCreation = (SColorPicker::OnColorPickerNonModalCreateOverride.IsBound() && !Args.bIsModal);
