ColorPicker: Change color sliders to respond to the sRGB preview checkbox

diff --git a/Engine/Source/Runtime/AppFramework/Private/Widgets/Colors/SColorPicker.cpp b/Engine/Source/Runtime/AppFramework/Private/Widgets/Colors/SColorPicker.cpp
--- a/Engine/Source/Runtime/AppFramework/Private/Widgets/Colors/SColorPicker.cpp
+++ b/Engine/Source/Runtime/AppFramework/Private/Widgets/Colors/SColorPicker.cpp
@@ -206,10 +206,19 @@ void SColorPicker::GenerateDefaultColorPickerContent( bool bAdvancedSectionExpan
 
 						+ SVerticalBox::Slot()
 							.AutoHeight()
+							[
+								// color preview
+								MakeColorPreviewBox()
+							]
+
+						+ SVerticalBox::Slot()
+							.AutoHeight()
+							.Padding(0.0f, 8.0f, 0.0f, 0.0f)
 							[
 								// sRGB check box
 								SNew(SCheckBox)
-									.ToolTipText(LOCTEXT("SRGBCheckboxToolTip", "Toggle gamma corrected sRGB previewing"))
+									.ToolTipText(LOCTEXT("SRGBCheckboxToolTip", "When enabled, the preview swatch uses sRGB encoding to correct the colors for display.\n"
+										"When disabled, the preview shows uncorrected linear colors."))
 									.IsChecked(this, &SColorPicker::HandleSRGBCheckBoxIsChecked)
 									.OnCheckStateChanged(this, &SColorPicker::HandleSRGBCheckBoxCheckStateChanged)
 									[
@@ -218,14 +227,6 @@ void SColorPicker::GenerateDefaultColorPickerContent( bool bAdvancedSectionExpan
 									]
 							]
 
-						+ SVerticalBox::Slot()
-							.AutoHeight()
-							.Padding(0.0f, 8.0f, 0.0f, 0.0f)
-							[
-								// color preview
-								MakeColorPreviewBox()
-							]
-
 						+ SVerticalBox::Slot()
 							.MinHeight(28.0f)
 							.MaxHeight(28.0f)
@@ -717,6 +718,7 @@ TSharedRef<SWidget> SColorPicker::MakeColorSlider( EColorPickerChannels Channel
 		.MaxSliderValue(Channel == EColorPickerChannels::Hue ? 359.0f : 1.0f)
 		.Delta(Channel == EColorPickerChannels::Hue ? 1.0f : 0.001f)
 		.HasAlphaBackground(Channel == EColorPickerChannels::Alpha)
+		.UseSRGB(this, &SColorPicker::HandleColorPickerUseSRGB)
 		.GradientColors(this, &SColorPicker::GetGradientColors, Channel)
 		.Value(this, &SColorPicker::HandleColorSpinBoxValue, Channel)
 		.Visibility(this, &SColorPicker::HandleColorPickerModeVisibility, EColorPickerModes::Wheel)
@@ -796,6 +798,7 @@ TSharedRef<SWidget> SColorPicker::MakeColorSpinBox( EColorPickerChannels Channel
 		.MaxSliderValue(Channel == EColorPickerChannels::Hue ? 359.0f : 1.0f)
 		.Delta(Channel == EColorPickerChannels::Hue ? 1.0f : 0.001f)
 		.HasAlphaBackground(Channel == EColorPickerChannels::Alpha)
+		.UseSRGB(this, &SColorPicker::HandleColorPickerUseSRGB)
 		.GradientColors(this, &SColorPicker::GetGradientColors, Channel)
 		.Value(this, &SColorPicker::HandleColorSpinBoxValue, Channel)
 		.OnBeginSliderMovement(const_cast<SColorPicker*>(this), &SColorPicker::HandleInteractiveChangeBegin)
@@ -902,7 +905,7 @@ FLinearColor SColorPicker::GetGradientEndColor( EColorPickerChannels Channel ) c
 	case EColorPickerChannels::Red:			return FLinearColor::Red;
 	case EColorPickerChannels::Green:		return FLinearColor::Green;
 	case EColorPickerChannels::Blue:		return FLinearColor::Blue;
-	case EColorPickerChannels::Alpha:		return FLinearColor::White;
+	case EColorPickerChannels::Alpha:		return FLinearColor(CurrentColorHSV.R, CurrentColorHSV.G, CurrentColorHSV.B, 1.0f).HSVToLinearRGB();
 	case EColorPickerChannels::Saturation:	return FLinearColor(CurrentColorHSV.R, 1.0f, 1.0f, 1.0f).HSVToLinearRGB();
 	case EColorPickerChannels::Value:		return FLinearColor(CurrentColorHSV.R, CurrentColorHSV.G, 1.0f, 1.0f).HSVToLinearRGB();
 	default:								return FLinearColor();
@@ -929,15 +932,23 @@ TArray<FLinearColor> SColorPicker::GetGradientColors( EColorPickerChannels Chann
 	TArray<FLinearColor> Colors;
 	if (Channel == EColorPickerChannels::Hue)
 	{
-		for (int32 i = 0; i < 7; ++i)
+		for (int32 i = 0; i < 10; ++i)
 		{
-			Colors.Add(FLinearColor((i % 6) * 60.f, 1.f, 1.f).HSVToLinearRGB());
+			Colors.Add(FLinearColor((i % 9) * 40.f, 1.f, 1.f).HSVToLinearRGB());
 		}
 	}
 	else
 	{
-		Colors.Add(GetGradientStartColor(Channel));
-		Colors.Add(GetGradientEndColor(Channel));
+		const FLinearColor StartColor = GetGradientStartColor(Channel);
+		const FLinearColor EndColor = GetGradientEndColor(Channel);
+
+		constexpr int32 NumSteps = 10;
+		constexpr float StepSize = 1.0f / NumSteps;
+		for (int32 Step = 0; Step <= NumSteps; ++Step)
+		{
+			const float Alpha = Step * StepSize;
+			Colors.Add(FMath::Lerp(StartColor, EndColor, Alpha));
+		}
 	}
 	return Colors;
 }
diff --git a/Engine/Source/Runtime/AppFramework/Private/Widgets/Colors/SColorSlider.cpp b/Engine/Source/Runtime/AppFramework/Private/Widgets/Colors/SColorSlider.cpp
--- a/Engine/Source/Runtime/AppFramework/Private/Widgets/Colors/SColorSlider.cpp
+++ b/Engine/Source/Runtime/AppFramework/Private/Widgets/Colors/SColorSlider.cpp
@@ -8,6 +8,7 @@ SColorSlider::SColorSlider()
 	: Orientation(*this, EOrientation::Orient_Horizontal)
 	, GradientColors(*this)
 	, bHasAlphaBackground(*this, false)
+	, bUseSRGB(*this, false)
 {
 }
 
@@ -16,6 +17,7 @@ void SColorSlider::Construct(const FArguments& InArgs)
 	GradientColors.Assign(*this, InArgs._GradientColors);
 
 	bHasAlphaBackground.Assign(*this, InArgs._HasAlphaBackground.Get());
+	bUseSRGB.Assign(*this, InArgs._UseSRGB);
 	Orientation.Assign(*this, InArgs._Orientation.Get());
 
 	BorderBrush = FAppStyle::Get().GetBrush("ColorPicker.RoundedInputBorder");
@@ -166,7 +168,8 @@ int32 SColorSlider::OnPaint(const FPaintArgs& Args, const FGeometry& AllottedGeo
 		TArray<FSlateGradientStop> GradientStops;
 		for (int32 ColorIndex = 0; ColorIndex < NumColors; ++ColorIndex)
 		{
-			GradientStops.Add(FSlateGradientStop(ColorSliderSize * (float(ColorIndex) / (NumColors - 1)), Colors[ColorIndex]));
+			FColor DrawColor = Colors[ColorIndex].ToFColor(bUseSRGB.Get());
+			GradientStops.Add(FSlateGradientStop(ColorSliderSize * (float(ColorIndex) / (NumColors - 1)), DrawColor));
 		}
 
 		const FVector4f CornerRadius = FVector4f(4.0f, 4.0f, 4.0f, 4.0f);
@@ -177,7 +180,7 @@ int32 SColorSlider::OnPaint(const FPaintArgs& Args, const FGeometry& AllottedGeo
 			AllottedGeometry.ToPaintGeometry(SliderSize, FSlateLayoutTransform(SliderOffset)),
 			GradientStops,
 			(Orientation.Get() == Orient_Horizontal) ? EOrientation::Orient_Vertical : EOrientation::Orient_Horizontal,
-			DrawEffects | ESlateDrawEffect::NoGamma, // Gradient colors are in linear space, so disable gamma to let them blend properly
+			DrawEffects,
 			CornerRadius
 		);	
 	}
diff --git a/Engine/Source/Runtime/AppFramework/Private/Widgets/Colors/SColorSlider.h b/Engine/Source/Runtime/AppFramework/Private/Widgets/Colors/SColorSlider.h
--- a/Engine/Source/Runtime/AppFramework/Private/Widgets/Colors/SColorSlider.h
+++ b/Engine/Source/Runtime/AppFramework/Private/Widgets/Colors/SColorSlider.h
@@ -33,6 +33,7 @@ public:
 		, _Label()
 		, _GradientColors()
 		, _HasAlphaBackground(false)
+		, _UseSRGB(true)
 		, _OnValueChanged()
 		, _OnBeginSliderMovement()
 		, _OnEndSliderMovement()
@@ -70,6 +71,9 @@ public:
 		/** Whether a checker background is displayed for alpha viewing */
 		SLATE_ATTRIBUTE(bool, HasAlphaBackground)
 
+		/** Whether to display sRGB color */
+		SLATE_ATTRIBUTE(bool, UseSRGB)
+
 		/** Called when the value is changed by the Slider or SpinBox */
 		SLATE_EVENT(FOnValueChanged, OnValueChanged)
 
@@ -114,6 +118,9 @@ private:
 	/** Whether a checker background is displayed for alpha viewing */
 	TSlateAttribute<bool, EInvalidateWidgetReason::Paint> bHasAlphaBackground;
 
+	/** Whether to display sRGB color */
+	TSlateAttribute<bool, EInvalidateWidgetReason::Paint> bUseSRGB;
+
 	/** Slider widget */
 	TSharedPtr<SSlider> Slider;
 
