ColorPicker: Replace hex linear/srgb text fields with a single field and a dropdown menu to pick between the two hex text modes

diff --git a/Engine/Source/Runtime/AppFramework/Private/Widgets/Colors/SColorPicker.cpp b/Engine/Source/Runtime/AppFramework/Private/Widgets/Colors/SColorPicker.cpp
--- a/Engine/Source/Runtime/AppFramework/Private/Widgets/Colors/SColorPicker.cpp
+++ b/Engine/Source/Runtime/AppFramework/Private/Widgets/Colors/SColorPicker.cpp
@@ -6,6 +6,7 @@
 #include "Layout/WidgetPath.h"
 #include "SlateOptMacros.h"
 #include "Framework/Application/SlateApplication.h"
+#include "Framework/MultiBox/MultiBoxBuilder.h"
 #include "Widgets/Layout/SBorder.h"
 #include "Widgets/Images/SImage.h"
 #include "Widgets/Text/STextBlock.h"
@@ -370,51 +371,30 @@ void SColorPicker::GenerateDefaultColorPickerContent( bool bAdvancedSectionExpan
 								MakeColorSpinBox(EColorPickerChannels::Value)
 							]
 
-						// Hexadecimal Linear TextBox
+						// Hexadecimal Dropdown and TextBox
 						+ SVerticalBox::Slot()
 							.MinHeight(20)
 							.MaxHeight(20)
 							.Padding(0.0f, 8.0f, 0.0f, 0.0f)
 							[
 								SNew(SHorizontalBox)
-									.ToolTipText(LOCTEXT("HexLinearSliderToolTip", "Hexadecimal Linear Value"))
 
 								+ SHorizontalBox::Slot()
 									.HAlign(HAlign_Left)
 									.VAlign(VAlign_Center)
 									[
-										SNew(STextBlock)
-											.Text(LOCTEXT("HexLinearInputLabel", "Hex Linear"))
-									]
-
-								+ SHorizontalBox::Slot()
-									.AutoWidth()
-									.HAlign(HAlign_Right)
-									[
-										SNew(SEditableTextBox)
-											.MinDesiredWidth(114.0f)
-											.Text(this, &SColorPicker::HandleHexLinearBoxText)
-											.Font(SmallFont)
-											.Padding(FMargin(8.0f, 4.0f, 8.0f, 4.0f))
-											.OnTextCommitted(this, &SColorPicker::HandleHexLinearInputTextCommitted)
-									]
-							]
-
-						// Hexadecimal sRGB TextBox
-						+ SVerticalBox::Slot()
-							.MinHeight(20)
-							.MaxHeight(20)
-							.Padding(0.0f, 8.0f, 0.0f, 0.0f)
-							[
-								SNew(SHorizontalBox)
-									.ToolTipText(LOCTEXT("HexSRGBSliderToolTip", "Hexadecimal sRGB Value"))
-
-								+ SHorizontalBox::Slot()
-									.HAlign(HAlign_Left)
-									.VAlign(VAlign_Center)
-									[
-										SNew(STextBlock)
-											.Text(LOCTEXT("HexSRGBInputLabel", "Hex sRGB"))
+										SNew(SComboButton)
+											.ComboButtonStyle(&FAppStyle::Get().GetWidgetStyle<FComboButtonStyle>("ColorPicker.HexMode"))
+											.MenuContent()
+											[
+												MakeHexModeMenu()
+											]
+											.ButtonContent()
+											[
+												SNew(STextBlock)
+													.Font(SmallFont)
+													.Text(this, &SColorPicker::HandleHexModeButtonText)
+											]
 									]
 
 								+ SHorizontalBox::Slot()
@@ -422,11 +402,11 @@ void SColorPicker::GenerateDefaultColorPickerContent( bool bAdvancedSectionExpan
 									.HAlign(HAlign_Right)
 									[
 										SNew(SEditableTextBox)
-											.MinDesiredWidth(114.0f)
-											.Text(this, &SColorPicker::HandleHexSRGBBoxText)
+											.MinDesiredWidth(109.0f)
+											.Text(this, &SColorPicker::HandleHexBoxText)
 											.Font(SmallFont)
 											.Padding(FMargin(8.0f, 4.0f, 8.0f, 4.0f))
-											.OnTextCommitted(this, &SColorPicker::HandleHexSRGBInputTextCommitted)
+											.OnTextCommitted(this, &SColorPicker::HandleHexInputTextCommitted)
 									]
 							]
 					]
@@ -1135,43 +1115,87 @@ void SColorPicker::HandleEyeDropperButtonComplete(bool bCancelled)
 }
 
 
-FText SColorPicker::HandleHexLinearBoxText() const
+FText SColorPicker::HandleHexBoxText() const
 {
-	return FText::FromString(CurrentColorRGB.ToFColor(false).ToHex());
+	const bool bSRGB = (HexMode == EColorPickerHexMode::SRGB);
+	return FText::FromString(CurrentColorRGB.ToFColor(bSRGB).ToHex());
 }
 
 
-FText SColorPicker::HandleHexSRGBBoxText() const
+void SColorPicker::HandleHexInputTextCommitted(const FText& Text, ETextCommit::Type CommitType)
 {
-	return FText::FromString(CurrentColorRGB.ToFColor(true).ToHex());
+	if (!Text.IsEmpty() && ((CommitType == ETextCommit::OnEnter) || (CommitType == ETextCommit::OnUserMovedFocus)))
+	{
+		FColor Color = FColor::FromHex(Text.ToString());
+		float Red = Color.R / 255.0f;
+		float Green = Color.G / 255.0f;
+		float Blue = Color.B / 255.0f;
+		float Alpha = Color.A / 255.0f;
+
+		if (HexMode == EColorPickerHexMode::SRGB)
+		{
+			Red = Red <= 0.04045f ? Red / 12.92f : FMath::Pow((Red + 0.055f) / 1.055f, 2.4f);
+			Green = Green <= 0.04045f ? Green / 12.92f : FMath::Pow((Green + 0.055f) / 1.055f, 2.4f);
+			Blue = Blue <= 0.04045f ? Blue / 12.92f : FMath::Pow((Blue + 0.055f) / 1.055f, 2.4f);
+		}
+
+		SetNewTargetColorRGB(FLinearColor(Red, Green, Blue, Alpha), false);
+	}
 }
 
 
-void SColorPicker::HandleHexLinearInputTextCommitted(const FText& Text, ETextCommit::Type CommitType)
+TSharedRef<SWidget> SColorPicker::MakeHexModeMenu()
 {
-	if (!Text.IsEmpty() && ((CommitType == ETextCommit::OnEnter) || (CommitType == ETextCommit::OnUserMovedFocus)))
-	{
-		FColor Color = FColor::FromHex(Text.ToString());
-		SetNewTargetColorRGB(FLinearColor(Color.R / 255.0f, Color.G / 255.0f, Color.B / 255.0f, Color.A / 255.0f), false);
-	}	
+	FMenuBuilder MenuBuilder(true, nullptr);
+
+	MenuBuilder.AddMenuEntry(
+		LOCTEXT("HexMenuText_SRGB", "Hex sRGB"),
+		LOCTEXT("HexMenuToolTip_SRGB", "Represents the color being created using sRGB encoding.\n"
+			"This format matches the hex color values typically used in web development and image editing software."),
+		FSlateIcon(),
+		FUIAction(
+			FExecuteAction::CreateLambda([this]() { HexMode = EColorPickerHexMode::SRGB; }),
+			FCanExecuteAction(),
+			FIsActionChecked::CreateLambda([this]() { return HexMode == EColorPickerHexMode::SRGB; })
+		),
+		NAME_None,
+		EUserInterfaceActionType::RadioButton
+	);
+
+	MenuBuilder.AddMenuEntry(
+		LOCTEXT("HexMenuText_Linear", "Hex Linear"),
+		LOCTEXT("HexMenuToolTip_Linear", "Represents the color being created using linear color values.\n"
+			"Note that linear hex values have less precision for darker colors."),
+		FSlateIcon(),
+		FUIAction(
+			FExecuteAction::CreateLambda([this]() { HexMode = EColorPickerHexMode::Linear; }),
+			FCanExecuteAction(),
+			FIsActionChecked::CreateLambda([this]() { return HexMode == EColorPickerHexMode::Linear; })
+		),
+		NAME_None,
+		EUserInterfaceActionType::RadioButton
+	);
+
+	return MenuBuilder.MakeWidget();
 }
 
-void SColorPicker::HandleHexSRGBInputTextCommitted(const FText& Text, ETextCommit::Type CommitType)
+FText SColorPicker::HandleHexModeButtonText() const
 {
-	if (!Text.IsEmpty() && ((CommitType == ETextCommit::OnEnter) || (CommitType == ETextCommit::OnUserMovedFocus)))
+	if (HexMode == EColorPickerHexMode::SRGB)
 	{
-		FColor Color = FColor::FromHex(Text.ToString());
-		float red = Color.R / 255.0f;
-		float green = Color.G / 255.0f;
-		float blue = Color.B / 255.0f;
-		float alpha = Color.A / 255.0f;
-
-		red = red <= 0.04045f ? red / 12.92f : FMath::Pow((red + 0.055f) / 1.055f, 2.4f);
-		green = green <= 0.04045f ? green / 12.92f : FMath::Pow((green + 0.055f) / 1.055f, 2.4f);
-		blue = blue <= 0.04045f ? blue / 12.92f : FMath::Pow((blue + 0.055f) / 1.055f, 2.4f);
-		
-		SetNewTargetColorRGB(FLinearColor(red, green, blue, alpha), false);
+		return LOCTEXT("HexMenuText_SRGB", "Hex sRGB");
 	}
+	else if (HexMode == EColorPickerHexMode::Linear)
+	{
+		return LOCTEXT("HexMenuText_Linear", "Hex Linear");
+	}
+
+	return FText::GetEmpty();
+}
+
+void SColorPicker::OnHexModeSelected(EColorPickerHexMode InHexMode)
+{
+	HexMode = InHexMode;
 }
 
 
diff --git a/Engine/Source/Runtime/AppFramework/Public/Widgets/Colors/SColorPicker.h b/Engine/Source/Runtime/AppFramework/Public/Widgets/Colors/SColorPicker.h
--- a/Engine/Source/Runtime/AppFramework/Public/Widgets/Colors/SColorPicker.h
+++ b/Engine/Source/Runtime/AppFramework/Public/Widgets/Colors/SColorPicker.h
@@ -49,6 +49,16 @@ enum class EColorPickerModes
 };
 
 
+/**
+ * Enumerates hexadecimal text modes.
+ */
+enum class EColorPickerHexMode
+{
+	SRGB,
+	Linear
+};
+
+
 /**
  * Struct for holding individual pointers to float values.
  */
@@ -323,16 +333,10 @@ private:
 	APPFRAMEWORK_API void HandleEyeDropperButtonComplete(bool bCancelled);
 
 	// Callback for getting the text in the hex linear box.
-	APPFRAMEWORK_API FText HandleHexLinearBoxText() const;
-
-	// Callback for getting the text in the hex sRGB box.
-	APPFRAMEWORK_API FText HandleHexSRGBBoxText() const;
+	APPFRAMEWORK_API FText HandleHexBoxText() const;
 
 	// Callback for committed text in the hex input box (sRGB gamma).
-	APPFRAMEWORK_API void HandleHexSRGBInputTextCommitted(const FText& Text, ETextCommit::Type CommitType);
-
-	// Callback for committed text in the hex input box (linear gamma).
-	APPFRAMEWORK_API void HandleHexLinearInputTextCommitted(const FText& Text, ETextCommit::Type CommitType);
+	APPFRAMEWORK_API void HandleHexInputTextCommitted(const FText& Text, ETextCommit::Type CommitType);
 
 	// Callback for changing the HSV value of the current color.
 	APPFRAMEWORK_API void HandleHSVColorChanged(FLinearColor NewValue);
@@ -388,6 +392,15 @@ private:
 	// Callback for selecting a color in the color theme bar.
 	APPFRAMEWORK_API void HandleThemeBarColorSelected(FLinearColor NewValue);
 
+	// Build the menu widget for selecting the hexadecimal text mode
+	APPFRAMEWORK_API TSharedRef<SWidget> MakeHexModeMenu();
+
+	// Get the text to display on the hexadecimal text mode button
+	APPFRAMEWORK_API FText HandleHexModeButtonText() const;
+
+	// Callback for when the hexadecimal text mode changes
+	APPFRAMEWORK_API void OnHexModeSelected(EColorPickerHexMode InHexMode);
+
 private:
 	
 	/** The color that is being targeted as a TAttribute */
@@ -411,6 +424,9 @@ private:
 	/** Holds the color picker's mode. */
 	EColorPickerModes CurrentMode;
 
+	/** Holds the color picker's hex mode. */
+	EColorPickerHexMode HexMode;
+
 	/** Time, used for color animation */
 	float CurrentTime;
 
diff --git a/Engine/Source/Runtime/SlateCore/Private/Styling/StarshipCoreStyle.cpp b/Engine/Source/Runtime/SlateCore/Private/Styling/StarshipCoreStyle.cpp
--- a/Engine/Source/Runtime/SlateCore/Private/Styling/StarshipCoreStyle.cpp
+++ b/Engine/Source/Runtime/SlateCore/Private/Styling/StarshipCoreStyle.cpp
@@ -1983,6 +1983,17 @@ void FStarshipCoreStyle::SetupColorPickerStyles(TSharedRef<FStyle>& Style)
 			.SetNormalThumbImage(IMAGE_BRUSH("Common/ColorPicker_SliderHandle", FVector2f(8.0f, 32.0f)))
 			.SetHoveredThumbImage(IMAGE_BRUSH("Common/ColorPicker_SliderHandle", FVector2f(8.0f, 32.0f)))
 		);
+
+		FButtonStyle ComboModeButtonStyle = FButtonStyle()
+			.SetNormal(FSlateRoundedBoxBrush(FStyleColors::Panel, InputFocusRadius, FStyleColors::Panel, InputFocusThickness))
+			.SetHovered(FSlateRoundedBoxBrush(FStyleColors::Panel, InputFocusRadius, FStyleColors::Panel, InputFocusThickness))
+			.SetPressed(FSlateRoundedBoxBrush(FStyleColors::Panel, InputFocusRadius, FStyleColors::Panel, InputFocusThickness));
+
+		FComboButtonStyle HexModeComboButton = Style->GetWidgetStyle<FComboButtonStyle>("ComboButton");
+		HexModeComboButton.SetContentPadding(FMargin(0.0f));
+		HexModeComboButton.SetButtonStyle(ComboModeButtonStyle);
+
+		Style->Set("ColorPicker.HexMode", HexModeComboButton);
 	}
 
 	// SColorSlider defaults...
