ColorPicker: Preserve color theme panel visibility state
#rnx

diff --git a/Engine/Source/Runtime/AppFramework/Private/Widgets/Colors/SColorPicker.cpp b/Engine/Source/Runtime/AppFramework/Private/Widgets/Colors/SColorPicker.cpp
--- a/Engine/Source/Runtime/AppFramework/Private/Widgets/Colors/SColorPicker.cpp
+++ b/Engine/Source/Runtime/AppFramework/Private/Widgets/Colors/SColorPicker.cpp
@@ -96,7 +96,8 @@ void SColorPicker::Construct( const FArguments& InArgs )
 
 		GConfig->GetBool(TEXT("ColorPickerUI"), TEXT("bWheelMode"), WheelMode, GEditorPerProjectIni);
 		GConfig->GetBool(TEXT("ColorPickerUI"), TEXT("bSRGBEnabled"), bUseSRGB, GEditorPerProjectIni);
-
+		GConfig->GetBool(TEXT("ColorPickerUI"), TEXT("bIsThemePanelVisible"), bIsThemePanelVisible, GEditorPerProjectIni);
+		
 		CurrentMode = WheelMode ? EColorPickerModes::Wheel : EColorPickerModes::Spectrum;
 	}
 
@@ -1305,6 +1306,12 @@ FReply SColorPicker::HandleOldColorPreviewClicked()
 FReply SColorPicker::ToggleThemePanelVisibility()
 {
 	bIsThemePanelVisible = !bIsThemePanelVisible;
+
+	if (FPaths::FileExists(GEditorPerProjectIni))
+	{
+		GConfig->SetBool(TEXT("ColorPickerUI"), TEXT("bIsThemePanelVisible"), bIsThemePanelVisible, GEditorPerProjectIni);
+	}
+
 	return FReply::Handled();
 }
 
