ColorPicker: Add maximum text length to theme name textbox to prevent incredibly long names that could break the menu UI.
#rnx

diff --git a/Engine/Source/Runtime/AppFramework/Private/Widgets/Colors/SColorThemes.cpp b/Engine/Source/Runtime/AppFramework/Private/Widgets/Colors/SColorThemes.cpp
--- a/Engine/Source/Runtime/AppFramework/Private/Widgets/Colors/SColorThemes.cpp
+++ b/Engine/Source/Runtime/AppFramework/Private/Widgets/Colors/SColorThemes.cpp
@@ -923,7 +923,8 @@ void SColorThemesViewer::Construct(const FArguments& InArgs)
 		.IsEnabled(false)
 		.OnTextCommitted(this, &SColorThemesViewer::CommitThemeName)
 		.ClearKeyboardFocusOnCommit(false)
-		.SelectAllTextWhenFocused(true);
+		.SelectAllTextWhenFocused(true)
+		.MaximumLength(128);
 
 	MultiBoxWidget = SNew(SMultiBoxWidget);
 	RefreshMenuWidget();
