ColorPicker: Fix stutter when adding/removing colors to a color theme.
#rnx

diff --git a/Engine/Source/Runtime/AppFramework/Private/Widgets/Colors/SColorThemes.cpp b/Engine/Source/Runtime/AppFramework/Private/Widgets/Colors/SColorThemes.cpp
--- a/Engine/Source/Runtime/AppFramework/Private/Widgets/Colors/SColorThemes.cpp
+++ b/Engine/Source/Runtime/AppFramework/Private/Widgets/Colors/SColorThemes.cpp
@@ -1363,8 +1363,6 @@ void SColorThemesViewer::SaveColorThemesToIni()
 			GConfig->SetString(TEXT("RecentColors"), *FString::Printf(TEXT("Color%i"), ColorIndex), *Color->ToString(), GEditorPerProjectIni);
 		}
 	}
-
-	GConfig->Flush(false, GEditorPerProjectIni);
 }
 
 #undef LOCTEXT_NAMESPACE
\ No newline at end of file
