ColorPicker: Fix ColorPicker to remember the user choice for the Hex mode (SRGB or Linear) when reopening the color picker window.
#rnx

diff --git a/Engine/Source/Runtime/AppFramework/Private/Widgets/Colors/SColorPicker.cpp b/Engine/Source/Runtime/AppFramework/Private/Widgets/Colors/SColorPicker.cpp
--- a/Engine/Source/Runtime/AppFramework/Private/Widgets/Colors/SColorPicker.cpp
+++ b/Engine/Source/Runtime/AppFramework/Private/Widgets/Colors/SColorPicker.cpp
@@ -93,12 +93,15 @@ void SColorPicker::Construct( const FArguments& InArgs )
 	if (FPaths::FileExists(GEditorPerProjectIni))
 	{
 		bool WheelMode = true;
+		bool bHexSRGB = true;
 
 		GConfig->GetBool(TEXT("ColorPickerUI"), TEXT("bWheelMode"), WheelMode, GEditorPerProjectIni);
 		GConfig->GetBool(TEXT("ColorPickerUI"), TEXT("bSRGBEnabled"), bUseSRGB, GEditorPerProjectIni);
+		GConfig->GetBool(TEXT("ColorPickerUI"), TEXT("bHexSRGB"), bHexSRGB, GEditorPerProjectIni);
 		GConfig->GetBool(TEXT("ColorPickerUI"), TEXT("bIsThemePanelVisible"), bIsThemePanelVisible, GEditorPerProjectIni);
 		
 		CurrentMode = WheelMode ? EColorPickerModes::Wheel : EColorPickerModes::Spectrum;
+		HexMode = bHexSRGB ? EColorPickerHexMode::SRGB : EColorPickerHexMode::Linear;
 	}
 
 	if (InArgs._sRGBOverride.IsSet())
@@ -1155,7 +1158,7 @@ TSharedRef<SWidget> SColorPicker::MakeHexModeMenu()
 			"This format matches the hex color values typically used in web development and image editing software."),
 		FSlateIcon(),
 		FUIAction(
-			FExecuteAction::CreateLambda([this]() { HexMode = EColorPickerHexMode::SRGB; }),
+			FExecuteAction::CreateSP(this, &SColorPicker::OnHexModeSelected, EColorPickerHexMode::SRGB),
 			FCanExecuteAction(),
 			FIsActionChecked::CreateLambda([this]() { return HexMode == EColorPickerHexMode::SRGB; })
 		),
@@ -1169,7 +1172,7 @@ TSharedRef<SWidget> SColorPicker::MakeHexModeMenu()
 			"Note that linear hex values have less precision for darker colors."),
 		FSlateIcon(),
 		FUIAction(
-			FExecuteAction::CreateLambda([this]() { HexMode = EColorPickerHexMode::Linear; }),
+			FExecuteAction::CreateSP(this, &SColorPicker::OnHexModeSelected, EColorPickerHexMode::Linear),
 			FCanExecuteAction(),
 			FIsActionChecked::CreateLambda([this]() { return HexMode == EColorPickerHexMode::Linear; })
 		),
@@ -1197,6 +1200,11 @@ FText SColorPicker::HandleHexModeButtonText() const
 void SColorPicker::OnHexModeSelected(EColorPickerHexMode InHexMode)
 {
 	HexMode = InHexMode;
+
+	if (FPaths::FileExists(GEditorPerProjectIni))
+	{
+		GConfig->SetBool(TEXT("ColorPickerUI"), TEXT("bHexSRGB"), HexMode == EColorPickerHexMode::SRGB, GEditorPerProjectIni);
+	}
 }
 
 
