ColorPicker: Fix Hue and Saturation sliders to allow them to accept values greater than the normal maximums (359.0f and 1.0 respecitvely).
When Saturation is allowed to go above 1.0f, some RGB values may go negative. The RGB color is now clamped so that none of the components can be negative.
#rnx

diff --git a/Engine/Source/Runtime/AppFramework/Private/Widgets/Colors/SColorPicker.cpp b/Engine/Source/Runtime/AppFramework/Private/Widgets/Colors/SColorPicker.cpp
--- a/Engine/Source/Runtime/AppFramework/Private/Widgets/Colors/SColorPicker.cpp
+++ b/Engine/Source/Runtime/AppFramework/Private/Widgets/Colors/SColorPicker.cpp
@@ -572,7 +572,7 @@ void SColorPicker::DiscardColor()
 bool SColorPicker::SetNewTargetColorHSV( const FLinearColor& NewValue, bool bForceUpdate /*= false*/ )
 {
 	CurrentColorHSV = NewValue;
-	CurrentColorRGB = NewValue.HSVToLinearRGB();
+	CurrentColorRGB = NewValue.HSVToLinearRGB().GetClamped(0.0f, FLT_MAX);
 
 	return ApplyNewTargetColor(bForceUpdate);
 }
@@ -580,7 +580,7 @@ bool SColorPicker::SetNewTargetColorHSV( const FLinearColor& NewValue, bool bFor
 
 bool SColorPicker::SetNewTargetColorRGB( const FLinearColor& NewValue, bool bForceUpdate /*= false*/ )
 {
-	CurrentColorRGB = NewValue;
+	CurrentColorRGB = NewValue.GetClamped(0.0f, FLT_MAX);
 	CurrentColorHSV = NewValue.LinearRGBToHSV();
 
 	return ApplyNewTargetColor(bForceUpdate);
@@ -699,8 +699,9 @@ TSharedRef<SWidget> SColorPicker::MakeColorSlider( EColorPickerChannels Channel
 		.Orientation(Orient_Vertical)
 		.ToolTipText(SliderTooltip)
 		.MinSliderValue(0.0f)
-		.MaxSliderValue(Channel == EColorPickerChannels::Hue ? 359.0f : 1.0f)
+		.MaxSliderValue(Channel == EColorPickerChannels::Hue ? 359.999f : 1.0f)
 		.Delta(Channel == EColorPickerChannels::Hue ? 1.0f : 0.001f)
+		.SupportDynamicSliderMaxValue(Channel == EColorPickerChannels::Hue ? false : true)
 		.HasAlphaBackground(Channel == EColorPickerChannels::Alpha)
 		.UseSRGB(this, &SColorPicker::HandleColorPickerUseSRGB)
 		.GradientColors(this, &SColorPicker::GetGradientColors, Channel)
@@ -752,13 +753,13 @@ TSharedRef<SWidget> SColorPicker::MakeColorSpinBox( EColorPickerChannels Channel
 		break;
 		
 	case EColorPickerChannels::Hue:
-		MaxValue = 359.0f;
+		MaxValue = HDRMaxValue;
 		SliderLabel = LOCTEXT("HueSliderLabel", "H");
 		SliderTooltip = LOCTEXT("HueSliderToolTip", "Hue");
 		break;
 		
 	case EColorPickerChannels::Saturation:
-		MaxValue = 1.0f;
+		MaxValue = HDRMaxValue;
 		SliderLabel = LOCTEXT("SaturationSliderLabel", "S");
 		SliderTooltip = LOCTEXT("SaturationSliderToolTip", "Saturation");
 		break;
@@ -779,8 +780,9 @@ TSharedRef<SWidget> SColorPicker::MakeColorSpinBox( EColorPickerChannels Channel
 		.MinSpinBoxValue(0.0f)
 		.MaxSpinBoxValue(MaxValue)
 		.MinSliderValue(0.0f)
-		.MaxSliderValue(Channel == EColorPickerChannels::Hue ? 359.0f : 1.0f)
+		.MaxSliderValue(Channel == EColorPickerChannels::Hue ? 359.999f : 1.0f)
 		.Delta(Channel == EColorPickerChannels::Hue ? 1.0f : 0.001f)
+		.SupportDynamicSliderMaxValue(Channel == EColorPickerChannels::Hue ? false : true)
 		.HasAlphaBackground(Channel == EColorPickerChannels::Alpha)
 		.UseSRGB(this, &SColorPicker::HandleColorPickerUseSRGB)
 		.GradientColors(this, &SColorPicker::GetGradientColors, Channel)
@@ -1074,7 +1076,7 @@ void SColorPicker::HandleColorSpinBoxValueChanged( float NewValue, EColorPickerC
 	case EColorPickerChannels::Green:		ComponentIndex = 1; break;
 	case EColorPickerChannels::Blue:		ComponentIndex = 2; break;
 	case EColorPickerChannels::Alpha:		ComponentIndex = 3; break;
-	case EColorPickerChannels::Hue:			ComponentIndex = 0; IsHSV = true; break;
+	case EColorPickerChannels::Hue:			ComponentIndex = 0; IsHSV = true; NewValue = FMath::Modulo(NewValue, 360.0f); break;
 	case EColorPickerChannels::Saturation:	ComponentIndex = 1; IsHSV = true; break;
 	case EColorPickerChannels::Value:		ComponentIndex = 2; IsHSV = true; break;
 	default:								
diff --git a/Engine/Source/Runtime/AppFramework/Private/Widgets/Colors/SColorSlider.cpp b/Engine/Source/Runtime/AppFramework/Private/Widgets/Colors/SColorSlider.cpp
--- a/Engine/Source/Runtime/AppFramework/Private/Widgets/Colors/SColorSlider.cpp
+++ b/Engine/Source/Runtime/AppFramework/Private/Widgets/Colors/SColorSlider.cpp
@@ -9,6 +9,7 @@ SColorSlider::SColorSlider()
 	, GradientColors(*this)
 	, bHasAlphaBackground(*this, false)
 	, bUseSRGB(*this, false)
+	, bSupportDynamicSliderMaxValue(*this, true)
 {
 }
 
@@ -18,6 +19,7 @@ void SColorSlider::Construct(const FArguments& InArgs)
 
 	bHasAlphaBackground.Assign(*this, InArgs._HasAlphaBackground.Get());
 	bUseSRGB.Assign(*this, InArgs._UseSRGB);
+	bSupportDynamicSliderMaxValue.Assign(*this, InArgs._SupportDynamicSliderMaxValue);
 	Orientation.Assign(*this, InArgs._Orientation.Get());
 
 	BorderBrush = FAppStyle::Get().GetBrush("ColorPicker.RoundedInputBorder");
@@ -85,7 +87,7 @@ void SColorSlider::Construct(const FArguments& InArgs)
 					.MaxFractionalDigits(3)
 					.Delta(InArgs._Delta)
 					.Value(InArgs._Value)
-					.SupportDynamicSliderMaxValue(true)
+					.SupportDynamicSliderMaxValue(InArgs._SupportDynamicSliderMaxValue)
 					.OnBeginSliderMovement(InArgs._OnBeginSpinBoxMovement)
 					.OnEndSliderMovement(InArgs._OnEndSpinBoxMovement)
 					.OnValueChanged(InArgs._OnValueChanged)
@@ -192,7 +194,7 @@ int32 SColorSlider::OnPaint(const FPaintArgs& Args, const FGeometry& AllottedGeo
 	// If the current slider value is greater than the slider's maximum value, update the slider maximum.
 	// This could occur because the maximum value of the spinbox might be greater than the maximum value of the slider. 
 	const float SliderValue = Slider->GetValue();
-	if (SliderValue > Slider->GetMaxValue())
+	if ((SliderValue > Slider->GetMaxValue()) && bSupportDynamicSliderMaxValue.Get())
 	{
 		Slider->SetMinAndMaxValues(Slider->GetMinValue(), SliderValue);
 	}
@@ -250,7 +252,7 @@ int32 SColorSlider::OnPaint(const FPaintArgs& Args, const FGeometry& AllottedGeo
 
 void SColorSlider::OnSpinBoxValueCommitted(float NewValue, ETextCommit::Type CommitType)
 {
-	if (NewValue > Slider->GetMaxValue())
+	if ((NewValue > Slider->GetMaxValue()) && bSupportDynamicSliderMaxValue.Get())
 	{
 		Slider->SetMinAndMaxValues(Slider->GetMinValue(), NewValue);
 	}
diff --git a/Engine/Source/Runtime/AppFramework/Private/Widgets/Colors/SColorSlider.h b/Engine/Source/Runtime/AppFramework/Private/Widgets/Colors/SColorSlider.h
--- a/Engine/Source/Runtime/AppFramework/Private/Widgets/Colors/SColorSlider.h
+++ b/Engine/Source/Runtime/AppFramework/Private/Widgets/Colors/SColorSlider.h
@@ -29,6 +29,7 @@ public:
 		, _MinSliderValue(0.0f)
 		, _MaxSliderValue(1.0f)
 		, _Delta(0.01f)
+		, _SupportDynamicSliderMaxValue(true)
 		, _Orientation(Orient_Horizontal)
 		, _Label()
 		, _GradientColors()
@@ -59,6 +60,9 @@ public:
 		/** The delta to increment the value as the slider moves */
 		SLATE_ATTRIBUTE(float, Delta)
 
+		/** Whether the underlying spinbox supports changing the maximum slider value */
+		SLATE_ATTRIBUTE(bool, SupportDynamicSliderMaxValue)
+
 		/** Orientation of the slider */
 		SLATE_ATTRIBUTE(EOrientation, Orientation)
 
@@ -121,6 +125,9 @@ private:
 	/** Whether to display sRGB color */
 	TSlateAttribute<bool, EInvalidateWidgetReason::Paint> bUseSRGB;
 
+	/** Whether to dynamically update the maximum slider value */
+	TSlateAttribute<bool, EInvalidateWidgetReason::Paint> bSupportDynamicSliderMaxValue;
+
 	/** Slider widget */
 	TSharedPtr<SSlider> Slider;
 
diff --git a/Engine/Source/Runtime/Slate/Private/Widgets/Colors/SColorWheel.cpp b/Engine/Source/Runtime/Slate/Private/Widgets/Colors/SColorWheel.cpp
--- a/Engine/Source/Runtime/Slate/Private/Widgets/Colors/SColorWheel.cpp
+++ b/Engine/Source/Runtime/Slate/Private/Widgets/Colors/SColorWheel.cpp
@@ -140,7 +140,7 @@ UE::Slate::FDeprecateVector2DResult SColorWheel::CalcRelativePositionFromCenter(
 	float Hue = SelectedColor.Get().R;
 	float Saturation = SelectedColor.Get().G;
 	float Angle = Hue / 180.0f * PI;
-	float Radius = Saturation;
+	float Radius = FMath::Clamp(Saturation, 0.0f, 1.0f);
 
 	return FVector2f(FMath::Cos(Angle), FMath::Sin(Angle)) * Radius;
 }
