[ColorTheme] Localize Copy append to duplicate color theme names

diff --git a/Engine/Source/Runtime/AppFramework/Private/Widgets/Colors/SColorThemes.cpp b/Engine/Source/Runtime/AppFramework/Private/Widgets/Colors/SColorThemes.cpp
--- a/Engine/Source/Runtime/AppFramework/Private/Widgets/Colors/SColorThemes.cpp
+++ b/Engine/Source/Runtime/AppFramework/Private/Widgets/Colors/SColorThemes.cpp
@@ -153,7 +153,7 @@ void SColorTrash::Construct( const FArguments& InArgs )
 	this->ChildSlot
 	[
 		SNew(SBorder)
-			.ToolTipText(NSLOCTEXT("ColorTrashWidget", "MouseOverToolTip", "Delete Color"))
+			.ToolTipText(LOCTEXT("MouseOverToolTip", "Delete Color"))
 			.BorderImage(this, &SColorTrash::GetBorderStyle)
 		[
 			SNew(SHorizontalBox)
@@ -344,7 +344,7 @@ void SThemeColorBlock::OnColorBlockRename()
 	// Field to enter new color label
 	TSharedRef<STextEntryPopup> TextEntry =
 		SNew(STextEntryPopup)
-		.Label(NSLOCTEXT("ThemeColorBlock","NewColorLabel", "Color Label"))
+		.Label(LOCTEXT("NewColorLabel", "Color Label"))
 		.OnTextCommitted(this, &SThemeColorBlock::SetLabel);
 
 	// Show dialog to enter new color label
@@ -1177,7 +1177,7 @@ TSharedPtr<FColorTheme> SColorThemesViewer::NewColorTheme(const FString& ThemeNa
 TSharedPtr<FColorTheme> SColorThemesViewer::GetDefaultColorTheme(bool bCreateNew)
 {
 	// Create a default theme (if bCreateNew is always creates a new one, even if there's already a like named theme)
-	const FText Name = NSLOCTEXT("ColorThemesViewer", "NewThemeName", "New Theme");
+	const FText Name = LOCTEXT("NewThemeName", "New Theme");
 	if ( bCreateNew )
 	{
 		return NewColorTheme( Name.ToString() );
@@ -1224,8 +1224,8 @@ void SColorThemesViewer::DuplicateColorTheme()
 	{
 		NewColors.Add(MakeShareable(new FColorInfo(CurrentColors[ColorIndex]->Color, CurrentColors[ColorIndex]->Label)));
 	}
-	const FText Name = NSLOCTEXT("ColorThemesViewer", "CopyThemeNameAppend", " Copy");
-	TSharedPtr<FColorTheme> NewTheme = NewColorTheme( CurrentlySelectedThemePtr.Pin()->Name + Name.ToString(), NewColors );
+	const FText Name = FText::Format(LOCTEXT("CopyThemeNameAppend", "{0} Copy"), FText::FromString(CurrentlySelectedThemePtr.Pin()->Name));
+	TSharedPtr<FColorTheme> NewTheme = NewColorTheme( Name.ToString(), NewColors );
 	SaveColorThemesToIni();
 
 	SetCurrentColorTheme(NewTheme);
