PR #11103: Fix color picker giving wrong color values during eye picking

diff --git a/Engine/Source/Runtime/AppFramework/Private/Widgets/Colors/SColorPicker.cpp b/Engine/Source/Runtime/AppFramework/Private/Widgets/Colors/SColorPicker.cpp
--- a/Engine/Source/Runtime/AppFramework/Private/Widgets/Colors/SColorPicker.cpp
+++ b/Engine/Source/Runtime/AppFramework/Private/Widgets/Colors/SColorPicker.cpp
@@ -54,6 +54,7 @@ SColorPicker::~SColorPicker()
 void SColorPicker::Construct( const FArguments& InArgs )
 {
 	TargetColorAttribute = InArgs._TargetColorAttribute;
+	ColorWheelBrush = InArgs._ColorWheelBrush;
 	CurrentColorHSV = OldColor = TargetColorAttribute.Get().LinearRGBToHSV();
 	CurrentColorRGB = TargetColorAttribute.Get();
 	CurrentMode = EColorPickerModes::Wheel;
@@ -160,8 +161,9 @@ void SColorPicker::GenerateDefaultColorPickerContent( bool bAdvancedSectionExpan
 											.MinWidth(200.0f)
 											.MaxWidth(200.0f)
 											[
-												SNew(SColorWheel)
+												SAssignNew(ColorWheel, SColorWheel)
 													.SelectedColor(this, &SColorPicker::GetCurrentColor)
+													.ColorWheelBrush(ColorWheelBrush)
 													.Visibility(this, &SColorPicker::HandleColorPickerModeVisibility, EColorPickerModes::Wheel)
 													.OnValueChanged(this, &SColorPicker::HandleColorWheelValueChanged)
 													.OnMouseCaptureBegin(this, &SColorPicker::HandleInteractiveChangeBegin)
@@ -276,7 +278,7 @@ void SColorPicker::GenerateDefaultColorPickerContent( bool bAdvancedSectionExpan
 										// eye dropper
 										SNew(SEyeDropperButton)
 											.OnValueChanged(this, &SColorPicker::HandleRGBColorChanged)
-											.OnBegin(this, &SColorPicker::HandleInteractiveChangeBegin)
+											.OnBegin(this, &SColorPicker::HandleEyeDropperButtonBegin)
 											.OnComplete(this, &SColorPicker::HandleEyeDropperButtonComplete)
 											.DisplayGamma(DisplayGamma)
 											.Visibility(bValidCreationOverrideExists ? EVisibility::Collapsed : EVisibility::Visible)
@@ -517,8 +519,9 @@ void SColorPicker::GenerateInlineColorPickerContent()
 		+ SHorizontalBox::Slot()
 			.AutoWidth()
 			[
-				SNew(SColorWheel)
+				SAssignNew(ColorWheel, SColorWheel)
 					.SelectedColor(this, &SColorPicker::GetCurrentColor)
+					.ColorWheelBrush(ColorWheelBrush)
 					.Visibility(this, &SColorPicker::HandleColorPickerModeVisibility, EColorPickerModes::Wheel)
 					.OnValueChanged(this, &SColorPicker::HandleColorWheelValueChanged)
 					.OnMouseCaptureBegin(this, &SColorPicker::HandleInteractiveChangeBegin)
@@ -1102,9 +1105,23 @@ void SColorPicker::HandleColorSpinBoxValueChanged( float NewValue, EColorPickerC
 	}
 }
 
+void SColorPicker::HandleEyeDropperButtonBegin()
+{
+	if (ColorWheel)
+	{
+		ColorWheel->ShowSelector(false);
+	}
+	
+	HandleInteractiveChangeBegin();
+}
 
 void SColorPicker::HandleEyeDropperButtonComplete(bool bCancelled)
 {
+	if (ColorWheel)
+	{
+		ColorWheel->ShowSelector(true);
+	}
+	
 	bIsInteractive = false;
 
 	if (bCancelled)
diff --git a/Engine/Source/Runtime/AppFramework/Public/Widgets/Colors/SColorPicker.h b/Engine/Source/Runtime/AppFramework/Public/Widgets/Colors/SColorPicker.h
--- a/Engine/Source/Runtime/AppFramework/Public/Widgets/Colors/SColorPicker.h
+++ b/Engine/Source/Runtime/AppFramework/Public/Widgets/Colors/SColorPicker.h
@@ -19,6 +19,7 @@ class SColorThemesViewer;
 class SComboButton;
 class SThemeColorBlocksBar;
 enum class EColorBlockAlphaDisplayMode : uint8;
+class SColorWheel;
 
 /** Called when the color picker cancel button is pressed */
 DECLARE_DELEGATE_OneParam(FOnColorPickerCancelled, FLinearColor);
@@ -87,6 +88,7 @@ public:
 
 	SLATE_BEGIN_ARGS(SColorPicker)
 		: _TargetColorAttribute(FLinearColor(ForceInit))
+		, _ColorWheelBrush(FCoreStyle::Get().GetBrush("ColorWheel.HueValueCircle"))
 		, _UseAlpha(true)
 		, _OnlyRefreshOnMouseUp(false)
 		, _OnlyRefreshOnOk(false)
@@ -109,6 +111,9 @@ public:
 		/** The color that is being targeted as a TAttribute */
 		SLATE_ATTRIBUTE(FLinearColor, TargetColorAttribute)
 
+		/** ColorWheelBrush to use */
+		SLATE_ATTRIBUTE(const FSlateBrush*, ColorWheelBrush)
+
 		/** An array of color pointers this color picker targets */
 		SLATE_ARGUMENT_DEPRECATED(TArray<FColor*>, TargetFColors, 5.2, "TargetFColors is deprecated. Use OnColorCommitted to get the selected color.")
 		
@@ -329,6 +334,9 @@ private:
 	// Callback for value changes in a color spin box.
 	APPFRAMEWORK_API void HandleColorSpinBoxValueChanged(float NewValue, EColorPickerChannels Channel);
 
+	// Callback for start of eye dropper interactions.
+	APPFRAMEWORK_API void HandleEyeDropperButtonBegin();
+
 	// Callback for completed eye dropper interactions.
 	APPFRAMEWORK_API void HandleEyeDropperButtonComplete(bool bCancelled);
 
@@ -406,6 +414,9 @@ private:
 	/** The color that is being targeted as a TAttribute */
 	TAttribute<FLinearColor> TargetColorAttribute;
 
+	/** The brush used as Color wheel as a TAttribute */
+	TAttribute<const FSlateBrush*> ColorWheelBrush;
+
 	/** The current color being picked in HSV */
 	FLinearColor CurrentColorHSV;
 
@@ -503,6 +514,9 @@ private:
 
 	/** Allows a details view to own the color picker so refreshing another details view doesn't close it */
 	TWeakPtr<SWidget> OptionalOwningDetailsView;
+
+	/** The Color Wheel widget embedded in this picker */
+	TSharedPtr<SColorWheel> ColorWheel;
 };
 
 
diff --git a/Engine/Source/Runtime/Slate/Private/Widgets/Colors/SColorWheel.cpp b/Engine/Source/Runtime/Slate/Private/Widgets/Colors/SColorWheel.cpp
--- a/Engine/Source/Runtime/Slate/Private/Widgets/Colors/SColorWheel.cpp
+++ b/Engine/Source/Runtime/Slate/Private/Widgets/Colors/SColorWheel.cpp
@@ -2,10 +2,10 @@
 
 #include "Widgets/Colors/SColorWheel.h"
 #include "Rendering/DrawElements.h"
-#include "Styling/CoreStyle.h"
 
 SColorWheel::SColorWheel()
 	: SelectedColor(*this, FLinearColor(ForceInit))
+	, bShouldDrawSelector(true)
 {}
 
 SColorWheel::~SColorWheel() = default;
@@ -15,7 +15,7 @@ SColorWheel::~SColorWheel() = default;
 
 void SColorWheel::Construct(const FArguments& InArgs)
 {
-	Image = FCoreStyle::Get().GetBrush("ColorWheel.HueValueCircle");
+	Image = InArgs._ColorWheelBrush.Get();
 	SelectorImage = FCoreStyle::Get().GetBrush("ColorWheel.Selector");
 	SelectedColor.Assign(*this, InArgs._SelectedColor);
 
@@ -116,19 +116,24 @@ int32 SColorWheel::OnPaint(const FPaintArgs& Args, const FGeometry& AllottedGeom
 		InWidgetStyle.GetColorAndOpacityTint() * Image->GetTint(InWidgetStyle)
 	);
 
-	// Draw Selector Image
-	const FVector2f SelectorImageLocation = CircleCenter + (CalcRelativePositionFromCenter() * CircleCenter) - (0.5f * SelectorSize);
-
-	FSlateDrawElement::MakeBox(
-		OutDrawElements,
-		LayerId + 1,
-		AllottedGeometry.ToPaintGeometry(SelectorSize, FSlateLayoutTransform(SelectorImageLocation)),
-		SelectorImage,
-		DrawEffects,
-		InWidgetStyle.GetColorAndOpacityTint() * SelectorImage->GetTint(InWidgetStyle)
-	);
+	if (bShouldDrawSelector)
+	{
+		// Draw Selector Image
+		const FVector2f SelectorImageLocation = CircleCenter + (CalcRelativePositionFromCenter() * CircleCenter) - (0.5f * SelectorSize);
+		
+		LayerId++;
+		FSlateDrawElement::MakeBox(
+			OutDrawElements,
+			LayerId,
+			AllottedGeometry.ToPaintGeometry(SelectorSize, FSlateLayoutTransform(SelectorImageLocation)),
+			SelectorImage,
+			DrawEffects,
+			InWidgetStyle.GetColorAndOpacityTint() * SelectorImage->GetTint(InWidgetStyle)
+		);
+	}
+	
 
-	return LayerId + 1;
+	return LayerId;
 }
 
 
diff --git a/Engine/Source/Runtime/Slate/Public/Widgets/Colors/SColorWheel.h b/Engine/Source/Runtime/Slate/Public/Widgets/Colors/SColorWheel.h
--- a/Engine/Source/Runtime/Slate/Public/Widgets/Colors/SColorWheel.h
+++ b/Engine/Source/Runtime/Slate/Public/Widgets/Colors/SColorWheel.h
@@ -8,6 +8,7 @@
 #include "Widgets/DeclarativeSyntaxSupport.h"
 #include "Widgets/SLeafWidget.h"
 #include "Framework/SlateDelegates.h"
+#include "Styling/CoreStyle.h"
 
 class FPaintArgs;
 class FSlateWindowElementList;
@@ -23,6 +24,7 @@ public:
 
 	SLATE_BEGIN_ARGS(SColorWheel)
 		: _SelectedColor()
+		, _ColorWheelBrush(FCoreStyle::Get().GetBrush("ColorWheel.HueValueCircle"))
 		, _OnMouseCaptureBegin()
 		, _OnMouseCaptureEnd()
 		, _OnValueChanged()
@@ -31,6 +33,9 @@ public:
 	
 		/** The current color selected by the user. */
 		SLATE_ATTRIBUTE(FLinearColor, SelectedColor)
+
+		/** ColorWheelBrush to use. */
+		SLATE_ATTRIBUTE(const FSlateBrush*, ColorWheelBrush)
 		
 		/** Invoked when the mouse is pressed and a capture begins. */
 		SLATE_EVENT(FSimpleDelegate, OnMouseCaptureBegin)
@@ -57,6 +62,13 @@ public:
 	 */
 	SLATE_API void Construct(const FArguments& InArgs);
 
+	/**
+	 * Set selector visibility.
+	 *
+	 * @param bShow
+	 */
+	SLATE_API void ShowSelector(bool bShow = true) { bShouldDrawSelector = bShow; }
+
 public:
 
 	//~ SWidget overrides
@@ -99,6 +111,9 @@ private:
 	/** The color selector image to show. */
 	const FSlateBrush* SelectorImage;
 
+	/** Flag to show/hide color selector */
+	bool bShouldDrawSelector;
+
 	/** Whether the user is dragging the slider */
 	bool bDragging = false;
 
