ColorPicker: Fix selected color theme not persisting after closing and reopening the editor

#rnx

diff --git a/Engine/Source/Runtime/AppFramework/Private/Widgets/Colors/SColorThemes.cpp b/Engine/Source/Runtime/AppFramework/Private/Widgets/Colors/SColorThemes.cpp
--- a/Engine/Source/Runtime/AppFramework/Private/Widgets/Colors/SColorThemes.cpp
+++ b/Engine/Source/Runtime/AppFramework/Private/Widgets/Colors/SColorThemes.cpp
@@ -2,6 +2,7 @@
 
 #include "Widgets/Colors/SColorThemes.h"
 
+#include "Algo/Find.h"
 #include "Application/SlateWindowHelper.h"
 #include "Framework/Application/SlateApplication.h"
 #include "Framework/MultiBox/MultiBoxBuilder.h"
@@ -1124,6 +1125,11 @@ void SColorThemesViewer::SetCurrentColorTheme(TSharedPtr<FColorTheme> NewTheme)
 	CurrentlySelectedThemePtr = NewTheme;
 	CurrentThemeChangedEvent.Broadcast();
 
+	if (FPaths::FileExists(GEditorPerProjectIni) && CurrentlySelectedThemePtr.IsValid())
+	{
+		GConfig->SetString(TEXT("ColorPickerUI"), TEXT("CurrentTheme"), *CurrentlySelectedThemePtr.Pin()->Name, GEditorPerProjectIni);
+	}
+
 	StopRename();
 }
 
@@ -1202,6 +1208,11 @@ void SColorThemesViewer::UpdateThemeNameFromTextBox()
 	{
 		GetCurrentColorTheme()->Name = MakeUniqueThemeName(Name);
 		SaveColorThemesToIni();
+
+		if (FPaths::FileExists(GEditorPerProjectIni) && CurrentlySelectedThemePtr.IsValid())
+		{
+			GConfig->SetString(TEXT("ColorPickerUI"), TEXT("CurrentTheme"), *CurrentlySelectedThemePtr.Pin()->Name, GEditorPerProjectIni);
+		}
 	}
 }
 
@@ -1328,6 +1339,17 @@ void SColorThemesViewer::LoadColorThemesFromIni()
 	 	// Create the default if none exists
 		GetDefaultColorTheme();
 	}
+
+	// Load the currently selected theme
+	if (FPaths::FileExists(GEditorPerProjectIni))
+	{
+		const FString ThemeName = GConfig->GetStr(TEXT("ColorPickerUI"), TEXT("CurrentTheme"), GEditorPerProjectIni);
+
+		if (TSharedPtr<FColorTheme>* CurrentTheme = Algo::FindBy(ColorThemes, ThemeName, &FColorTheme::Name))
+		{
+			CurrentlySelectedThemePtr = *CurrentTheme;
+		}
+	}
 }
 
 void SColorThemesViewer::SaveColorThemesToIni()
