Fix crash when opening color picker during a hitch

If the widget hierarchy has changed but the hittest grid has not yet updated (due to the hitch), GeneratePathToWidgetChecked in OpenColorPicker will fail, resulting in an assert and then crash. This changes OpenColorPicker to use GeneratePathToWidgetUnchecked instead.

#tests Verified in City Sample that the issue no longer occurs with this change.

diff --git a/Engine/Source/Runtime/AppFramework/Private/Widgets/Colors/SColorPicker.cpp b/Engine/Source/Runtime/AppFramework/Private/Widgets/Colors/SColorPicker.cpp
--- a/Engine/Source/Runtime/AppFramework/Private/Widgets/Colors/SColorPicker.cpp
+++ b/Engine/Source/Runtime/AppFramework/Private/Widgets/Colors/SColorPicker.cpp
@@ -1538,6 +1538,7 @@ bool OpenColorPicker(const FColorPickerArgs& Args)
 		if (Args.bIsModal)
 		{
 			FSlateApplication::Get().AddModalWindow(Window.ToSharedRef(), Args.ParentWidget);
+			Result = true;
 		}
 		else if (bNeedToAddWindow)
 		{
@@ -1545,18 +1546,19 @@ bool OpenColorPicker(const FColorPickerArgs& Args)
 			{
 				// Find the window of the parent widget
 				FWidgetPath WidgetPath;
-				FSlateApplication::Get().GeneratePathToWidgetChecked(Args.ParentWidget.ToSharedRef(), WidgetPath);
-				Window = FSlateApplication::Get().AddWindowAsNativeChild(Window.ToSharedRef(), WidgetPath.GetWindow());
+				if (FSlateApplication::Get().GeneratePathToWidgetUnchecked(Args.ParentWidget.ToSharedRef(), WidgetPath))
+				{
+					Window = FSlateApplication::Get().AddWindowAsNativeChild(Window.ToSharedRef(), WidgetPath.GetWindow());
+					Result = true;
+				}
 			}
 			else
 			{
 				Window = FSlateApplication::Get().AddWindow(Window.ToSharedRef());
+				Result = true;
 			}
-
 		}
 
-		Result = true;
-
 		//hold on to the window created for external use...
 		ColorPickerWindow = Window;
 	}
