[Starship Suite] Expose register/unregister tab spawner functions for StarshipGallery, as well as the name of the tab. Allows for embedding into existing windows/tab layouts.

diff --git a/Engine/Source/Runtime/AppFramework/Private/Widgets/Testing/SStarshipSuite.cpp b/Engine/Source/Runtime/AppFramework/Private/Widgets/Testing/SStarshipSuite.cpp
--- a/Engine/Source/Runtime/AppFramework/Private/Widgets/Testing/SStarshipSuite.cpp
+++ b/Engine/Source/Runtime/AppFramework/Private/Widgets/Testing/SStarshipSuite.cpp
@@ -19,12 +19,10 @@
 
 #include "Brushes/SlateDynamicImageBrush.h"
 
-void RestoreStarshipSuite()
-{
-
-	// Need to load this module so we have the widget reflector tab available
-	FModuleManager::LoadModuleChecked<ISlateReflectorModule>("SlateReflector");
+const FLazyName StarshipGalleryTabName("STARSHIP GALLERY");
 
+void RegisterStarshipGalleryTabSpawner()
+{
 	auto SpawnGallery = [](const FSpawnTabArgs&) -> TSharedRef<SDockTab>
 	{
 		return SNew(SDockTab)
@@ -33,7 +31,21 @@ void RestoreStarshipSuite()
 				MakeStarshipGallery()
 			];
 	};
-	FGlobalTabmanager::Get()->RegisterNomadTabSpawner("STARSHIP GALLERY", FOnSpawnTab::CreateLambda(SpawnGallery));
+	FGlobalTabmanager::Get()->RegisterNomadTabSpawner(StarshipGalleryTabName, FOnSpawnTab::CreateLambda(SpawnGallery));
+}
+
+void UnregisterStarshipGalleryTabSpawner()
+{
+	FGlobalTabmanager::Get()->UnregisterNomadTabSpawner(StarshipGalleryTabName);
+}
+
+void RestoreStarshipSuite()
+{
+
+	// Need to load this module so we have the widget reflector tab available
+	FModuleManager::LoadModuleChecked<ISlateReflectorModule>("SlateReflector");
+
+	RegisterStarshipGalleryTabSpawner();
 	
 	TSharedRef<FTabManager::FLayout> Layout = FTabManager::NewLayout( "StarshipSuite_Layout" )
 	->AddArea
@@ -42,9 +54,9 @@ void RestoreStarshipSuite()
 		->Split
 		(
 			FTabManager::NewStack()
-			->AddTab("STARSHIP GALLERY", ETabState::OpenedTab)
+			->AddTab(StarshipGalleryTabName, ETabState::OpenedTab)
 			->AddTab("WidgetReflector", ETabState::OpenedTab)
-			->SetForegroundTab(FName("STARSHIP GALLERY"))
+			->SetForegroundTab(FName(StarshipGalleryTabName))
 		)
 	);
 
diff --git a/Engine/Source/Runtime/AppFramework/Public/Widgets/Testing/SStarshipSuite.h b/Engine/Source/Runtime/AppFramework/Public/Widgets/Testing/SStarshipSuite.h
--- a/Engine/Source/Runtime/AppFramework/Public/Widgets/Testing/SStarshipSuite.h
+++ b/Engine/Source/Runtime/AppFramework/Public/Widgets/Testing/SStarshipSuite.h
@@ -7,6 +7,9 @@
 
 #if !UE_BUILD_SHIPPING
 
+extern APPFRAMEWORK_API const FLazyName StarshipGalleryTabName;
+APPFRAMEWORK_API void RegisterStarshipGalleryTabSpawner();
+APPFRAMEWORK_API void UnregisterStarshipGalleryTabSpawner();
 APPFRAMEWORK_API void RestoreStarshipSuite();
 
 #endif // #if !UE_BUILD_SHIPPING
